Changed display name of MassCrowd plugin from "Crowd" to "MassCrowd" for consistency with other Mass-family plugins

diff --git a/Engine/Plugins/AI/MassCrowd/MassCrowd.uplugin b/Engine/Plugins/AI/MassCrowd/MassCrowd.uplugin
--- a/Engine/Plugins/AI/MassCrowd/MassCrowd.uplugin
+++ b/Engine/Plugins/AI/MassCrowd/MassCrowd.uplugin
@@ -2,7 +2,7 @@
 	"FileVersion": 1,
 	"Version": 1,
 	"VersionName": "1.0",
-	"FriendlyName": "Crowd",
+	"FriendlyName": "MassCrowd",
 	"Description": "Spline based AI crowd system",
 	"Category": "AI",
 	"CreatedBy" : "Epic Games Inc",
