Fixes for MassSmartObject map in AIProject

#rnx

diff --git a/Engine/Plugins/AI/MassCrowd/Source/MassCrowd/Private/MassCrowdVisualizationLODProcessor.cpp b/Engine/Plugins/AI/MassCrowd/Source/MassCrowd/Private/MassCrowdVisualizationLODProcessor.cpp
--- a/Engine/Plugins/AI/MassCrowd/Source/MassCrowd/Private/MassCrowdVisualizationLODProcessor.cpp
+++ b/Engine/Plugins/AI/MassCrowd/Source/MassCrowd/Private/MassCrowdVisualizationLODProcessor.cpp
@@ -39,6 +39,8 @@ namespace UE::MassCrowd
 
 UMassCrowdVisualizationLODProcessor::UMassCrowdVisualizationLODProcessor()
 {
+	bAutoRegisterWithProcessingPhases = true;
+
 	ExecutionFlags = (int32)(EProcessorExecutionFlags::Client | EProcessorExecutionFlags::Standalone);
 
 	ExecutionOrder.ExecuteInGroup = UE::Mass::ProcessorGroupNames::LOD;
diff --git a/Engine/Plugins/AI/MassCrowd/Source/MassCrowd/Public/MassCrowdRepresentationProcessor.h b/Engine/Plugins/AI/MassCrowd/Source/MassCrowd/Public/MassCrowdRepresentationProcessor.h
--- a/Engine/Plugins/AI/MassCrowd/Source/MassCrowd/Public/MassCrowdRepresentationProcessor.h
+++ b/Engine/Plugins/AI/MassCrowd/Source/MassCrowd/Public/MassCrowdRepresentationProcessor.h
@@ -10,7 +10,7 @@
  * Overridden representation processor to make it tied to the crowd via the requirements.
  * It is also the base class for all the different type of crowd representation (Visualization & ServerSideRepresentation)
  */
-UCLASS(meta = (DisplayName = "Mass Crowd Visualization"))
+UCLASS(abstract)
 class MASSCROWD_API UMassCrowdRepresentationProcessor : public UMassRepresentationProcessor
 {
 	GENERATED_BODY()
