Changes to representation-related fragment destructors to differentiate between default and crowd-specific destructor. This change is done in preparation for new de/initializers

diff --git a/Engine/Plugins/AI/MassCrowd/Source/MassCrowd/Private/MassCrowdVisualizationProcessor.cpp b/Engine/Plugins/AI/MassCrowd/Source/MassCrowd/Private/MassCrowdVisualizationProcessor.cpp
--- a/Engine/Plugins/AI/MassCrowd/Source/MassCrowd/Private/MassCrowdVisualizationProcessor.cpp
+++ b/Engine/Plugins/AI/MassCrowd/Source/MassCrowd/Private/MassCrowdVisualizationProcessor.cpp
@@ -21,6 +21,9 @@ namespace UE::MassCrowd
 	};
 } // UE::MassCrowd
 
+//----------------------------------------------------------------------//
+// UMassCrowdVisualizationProcessor
+//----------------------------------------------------------------------//
 UMassCrowdVisualizationProcessor::UMassCrowdVisualizationProcessor()
 {
 	ExecutionFlags = (int32)(EProcessorExecutionFlags::Client | EProcessorExecutionFlags::Standalone);
@@ -79,6 +82,9 @@ void UMassCrowdVisualizationProcessor::Execute(UMassEntitySubsystem& EntitySubsy
 	}
 }
 
+//----------------------------------------------------------------------//
+// UMassCrowdRepresentationFragmentDestructor
+//----------------------------------------------------------------------//
 void UMassCrowdRepresentationFragmentDestructor::Initialize(UObject& Owner)
 {
 	Super::Initialize(Owner);
@@ -86,3 +92,10 @@ void UMassCrowdRepresentationFragmentDestructor::Initialize(UObject& Owner)
 	// Override default representation subsystem by the crowd one to support parallelization
 	RepresentationSubsystem = UWorld::GetSubsystem<UMassCrowdRepresentationSubsystem>(Owner.GetWorld());
 }
+
+void UMassCrowdRepresentationFragmentDestructor::ConfigureQueries()
+{
+	EntityQuery.AddRequirement<FMassRepresentationFragment>(EMassFragmentAccess::ReadOnly);
+	EntityQuery.AddRequirement<FDataFragment_Actor>(EMassFragmentAccess::ReadWrite);
+	EntityQuery.AddTagRequirement<FMassCrowdRepresentationDestructorTag>(EMassFragmentPresence::All);
+}
\ No newline at end of file
diff --git a/Engine/Plugins/AI/MassCrowd/Source/MassCrowd/Private/MassCrowdVisualizationTrait.cpp b/Engine/Plugins/AI/MassCrowd/Source/MassCrowd/Private/MassCrowdVisualizationTrait.cpp
--- a/Engine/Plugins/AI/MassCrowd/Source/MassCrowd/Private/MassCrowdVisualizationTrait.cpp
+++ b/Engine/Plugins/AI/MassCrowd/Source/MassCrowd/Private/MassCrowdVisualizationTrait.cpp
@@ -8,5 +8,7 @@ UMassCrowdVisualizationTrait::UMassCrowdVisualizationTrait()
 {
 	// Override the subsystem and deinitializer to support parallelization of the crowd
 	RepresentationSubsystemClass = UMassCrowdRepresentationSubsystem::StaticClass();
+	// @todo the following line will be cut once new de/initializers are in
 	RepresentationFragmentDeinitializerClass = UMassCrowdRepresentationFragmentDestructor::StaticClass();
+	RepresentationDestructorTag = FMassCrowdRepresentationDestructorTag::StaticStruct();
 }
diff --git a/Engine/Plugins/AI/MassCrowd/Source/MassCrowd/Public/MassCrowdVisualizationProcessor.h b/Engine/Plugins/AI/MassCrowd/Source/MassCrowd/Public/MassCrowdVisualizationProcessor.h
--- a/Engine/Plugins/AI/MassCrowd/Source/MassCrowd/Public/MassCrowdVisualizationProcessor.h
+++ b/Engine/Plugins/AI/MassCrowd/Source/MassCrowd/Public/MassCrowdVisualizationProcessor.h
@@ -29,10 +29,17 @@ protected:
 	virtual void Execute(UMassEntitySubsystem& EntitySubsystem, FMassExecutionContext& Context) override;
 };
 
+USTRUCT()
+struct FMassCrowdRepresentationDestructorTag : public FMassTag
+{
+	GENERATED_BODY()
+};
+
 UCLASS()
 class MASSCROWD_API UMassCrowdRepresentationFragmentDestructor : public UMassRepresentationFragmentDestructor
 {
 	GENERATED_BODY()
 protected:
 	virtual void Initialize(UObject& Owner) override;
+	virtual void ConfigureQueries() override;
 };
diff --git a/Engine/Plugins/Runtime/MassGameplay/Source/MassRepresentation/Private/MassRepresentationProcessor.cpp b/Engine/Plugins/Runtime/MassGameplay/Source/MassRepresentation/Private/MassRepresentationProcessor.cpp
--- a/Engine/Plugins/Runtime/MassGameplay/Source/MassRepresentation/Private/MassRepresentationProcessor.cpp
+++ b/Engine/Plugins/Runtime/MassGameplay/Source/MassRepresentation/Private/MassRepresentationProcessor.cpp
@@ -497,16 +497,19 @@ void UMassRepresentationFragmentDestructor::Initialize(UObject& Owner)
 	Super::Initialize(Owner);
 
 	RepresentationSubsystem = UWorld::GetSubsystem<UMassRepresentationSubsystem>(Owner.GetWorld());
+	check(RepresentationSubsystem);
 }
 
 void UMassRepresentationFragmentDestructor::ConfigureQueries()
 {
 	EntityQuery.AddRequirement<FMassRepresentationFragment>(EMassFragmentAccess::ReadOnly);
 	EntityQuery.AddRequirement<FDataFragment_Actor>(EMassFragmentAccess::ReadWrite);
+	EntityQuery.AddTagRequirement<FMassRepresentationDefaultDestructorTag>(EMassFragmentPresence::All);
 }
 
 void UMassRepresentationFragmentDestructor::Execute(UMassEntitySubsystem& EntitySubsystem, FMassExecutionContext& Context)
 {
+	check(Context.ValidateAuxDataType<FMassRepresentationFragment>());
 	check(RepresentationSubsystem);
 
 	EntityQuery.ForEachEntityChunk(EntitySubsystem, Context, [this](FMassExecutionContext& Context)
diff --git a/Engine/Plugins/Runtime/MassGameplay/Source/MassRepresentation/Private/MassVisualizationTrait.cpp b/Engine/Plugins/Runtime/MassGameplay/Source/MassRepresentation/Private/MassVisualizationTrait.cpp
--- a/Engine/Plugins/Runtime/MassGameplay/Source/MassRepresentation/Private/MassVisualizationTrait.cpp
+++ b/Engine/Plugins/Runtime/MassGameplay/Source/MassRepresentation/Private/MassVisualizationTrait.cpp
@@ -11,7 +11,9 @@
 UMassVisualizationTrait::UMassVisualizationTrait()
 {
 	RepresentationSubsystemClass = UMassRepresentationSubsystem::StaticClass();
+	// @todo the following line will be cut once new de/initializers are in
 	RepresentationFragmentDeinitializerClass = UMassRepresentationFragmentDestructor::StaticClass();
+	RepresentationDestructorTag = FMassRepresentationDefaultDestructorTag::StaticStruct();
 }
 
 void UMassVisualizationTrait::BuildTemplate(FMassEntityTemplateBuildContext& BuildContext, UWorld& World) const
@@ -34,7 +36,12 @@ void UMassVisualizationTrait::BuildTemplate(FMassEntityTemplateBuildContext& Bui
 		RepresentationFragment.HighResTemplateActorIndex = HighResTemplateActor.Get() ? RepresentationSubsystem->FindOrAddTemplateActor(HighResTemplateActor.Get()) : INDEX_NONE;
 		RepresentationFragment.LowResTemplateActorIndex = LowResTemplateActor.Get() ? RepresentationSubsystem->FindOrAddTemplateActor(LowResTemplateActor.Get()) : INDEX_NONE;
 
+		// @todo the following line will be cut once new de/initializers are in
 		BuildContext.AddDeinitializer(*(RepresentationFragmentDeinitializerClass->GetDefaultObject<UMassProcessor>()));
+		if (ensureMsgf(RepresentationDestructorTag, TEXT("RepresentationDestructorTag is never expected to be empty")))
+		{
+			BuildContext.AddTag(*RepresentationDestructorTag);
+		}
 
 		BuildContext.AddFragment<FMassViewerInfoFragment>();
 		BuildContext.AddFragment<FMassRepresentationLODFragment>();
diff --git a/Engine/Plugins/Runtime/MassGameplay/Source/MassRepresentation/Public/MassRepresentationProcessor.h b/Engine/Plugins/Runtime/MassGameplay/Source/MassRepresentation/Public/MassRepresentationProcessor.h
--- a/Engine/Plugins/Runtime/MassGameplay/Source/MassRepresentation/Public/MassRepresentationProcessor.h
+++ b/Engine/Plugins/Runtime/MassGameplay/Source/MassRepresentation/Public/MassRepresentationProcessor.h
@@ -208,6 +208,13 @@ protected:
 	float NotVisibleUpdateRate = 0.5f;
 };
 
+
+USTRUCT()
+struct FMassRepresentationDefaultDestructorTag : public FMassTag
+{
+	GENERATED_BODY()
+};
+
 UCLASS()
 class MASSREPRESENTATION_API UMassRepresentationFragmentDestructor : public UMassFragmentDestructor
 {
diff --git a/Engine/Plugins/Runtime/MassGameplay/Source/MassRepresentation/Public/MassVisualizationTrait.h b/Engine/Plugins/Runtime/MassGameplay/Source/MassRepresentation/Public/MassVisualizationTrait.h
--- a/Engine/Plugins/Runtime/MassGameplay/Source/MassRepresentation/Public/MassVisualizationTrait.h
+++ b/Engine/Plugins/Runtime/MassGameplay/Source/MassRepresentation/Public/MassVisualizationTrait.h
@@ -34,9 +34,13 @@ protected:
 	TSubclassOf<AActor> LowResTemplateActor;
 
 	/** Allow subclasses to override the representation subsystem to use */
+	UPROPERTY()
 	TSubclassOf<UMassRepresentationSubsystem> RepresentationSubsystemClass;
 
+	// @todo the following property will be cut once new de/initializers are in
 	/** Allow subclasses to override the representation fragment deinitializer */
 	TSubclassOf<UMassProcessor> RepresentationFragmentDeinitializerClass;
 
+	UPROPERTY()
+	UScriptStruct* RepresentationDestructorTag;
 };
