Removed InitializeVelocity as it is already done in the translator

#rnx

diff --git a/Engine/Plugins/AI/MassCrowd/Source/MassCrowd/Private/MassCrowdRepresentationProcessor.cpp b/Engine/Plugins/AI/MassCrowd/Source/MassCrowd/Private/MassCrowdRepresentationProcessor.cpp
--- a/Engine/Plugins/AI/MassCrowd/Source/MassCrowd/Private/MassCrowdRepresentationProcessor.cpp
+++ b/Engine/Plugins/AI/MassCrowd/Source/MassCrowd/Private/MassCrowdRepresentationProcessor.cpp
@@ -38,37 +38,6 @@ void UMassCrowdRepresentationProcessor::ConfigureQueries()
 	CharacterMovementEntitiesQuery_Conditional.SetChunkFilter(&FMassVisualizationChunkFragment::AreAnyEntitiesVisibleInChunk);
 }
 
-void UMassCrowdRepresentationProcessor::InitializeVelocity(UMassEntitySubsystem& EntitySubsystem, FMassExecutionContext& Context)
-{
-	CharacterMovementEntitiesQuery_Conditional.ForEachEntityChunk(EntitySubsystem, Context, [this](FMassExecutionContext& Context)
-		{
-			const TConstArrayView<FMassRepresentationFragment> VisualizationList = Context.GetFragmentView<FMassRepresentationFragment>();
-			const TConstArrayView<FMassVelocityFragment> VelocityList = Context.GetFragmentView<FMassVelocityFragment>();
-			const TArrayView<FDataFragment_CharacterMovementComponentWrapper> CharMoveWrapperList = Context.GetMutableFragmentView<FDataFragment_CharacterMovementComponentWrapper>();
-
-			const int32 NumEntities = Context.GetNumEntities();
-			for (int32 EntityIdx = 0; EntityIdx < NumEntities; EntityIdx++)
-			{
-				const FMassRepresentationFragment& Visualization = VisualizationList[EntityIdx];
-
-				// @todo: Add syncing skeletal mesh <-> static mesh here.
-				if (Visualization.PrevRepresentation != Visualization.CurrentRepresentation)
-				{
-					const FMassVelocityFragment& Velocity = VelocityList[EntityIdx];
-					FDataFragment_CharacterMovementComponentWrapper& CharMoveWrapper = CharMoveWrapperList[EntityIdx];
-					UCharacterMovementComponent* MovementComp = CharMoveWrapper.Component.Get();
-
-					const bool bWasUsingMovementComponent = (Visualization.PrevRepresentation == ERepresentationType::LowResSpawnedActor || Visualization.PrevRepresentation == ERepresentationType::HighResSpawnedActor);
-					const bool bIsUsingMovementComponent = (Visualization.CurrentRepresentation == ERepresentationType::LowResSpawnedActor || Visualization.CurrentRepresentation == ERepresentationType::HighResSpawnedActor);
-					if (MovementComp && !bWasUsingMovementComponent && bIsUsingMovementComponent)
-					{
-						MovementComp->Velocity = Velocity.Value;
-					}
-				}
-			}
-		});
-}
-
 void UMassCrowdRepresentationProcessor::SetActorEnabled(const EActorEnabledType EnabledType, AActor& Actor, const int32 EntityIdx, FMassCommandBuffer& CommandBuffer)
 {
 	Super::SetActorEnabled(EnabledType, Actor, EntityIdx, CommandBuffer);
diff --git a/Engine/Plugins/AI/MassCrowd/Source/MassCrowd/Private/MassCrowdServerRepresentationProcessor.cpp b/Engine/Plugins/AI/MassCrowd/Source/MassCrowd/Private/MassCrowdServerRepresentationProcessor.cpp
--- a/Engine/Plugins/AI/MassCrowd/Source/MassCrowd/Private/MassCrowdServerRepresentationProcessor.cpp
+++ b/Engine/Plugins/AI/MassCrowd/Source/MassCrowd/Private/MassCrowdServerRepresentationProcessor.cpp
@@ -14,15 +14,4 @@ UMassCrowdServerRepresentationProcessor::UMassCrowdServerRepresentationProcessor
 	LODRepresentation[EMassLOD::Medium] = ERepresentationType::None;
 	LODRepresentation[EMassLOD::Low] = ERepresentationType::None;
 	LODRepresentation[EMassLOD::Off] = ERepresentationType::None;
-}
-
-void UMassCrowdServerRepresentationProcessor::Execute(UMassEntitySubsystem& EntitySubsystem, FMassExecutionContext& Context)
-{
-	EntityQuery.ForEachEntityChunk(EntitySubsystem, Context, [this](FMassExecutionContext& Context)
-	{
-		UpdateRepresentation(Context);
-	});
-
-	// @todo, we should use the new translators to do that initialization
-	InitializeVelocity(EntitySubsystem, Context);
 }
\ No newline at end of file
diff --git a/Engine/Plugins/AI/MassCrowd/Source/MassCrowd/Private/MassCrowdVisualizationProcessor.cpp b/Engine/Plugins/AI/MassCrowd/Source/MassCrowd/Private/MassCrowdVisualizationProcessor.cpp
--- a/Engine/Plugins/AI/MassCrowd/Source/MassCrowd/Private/MassCrowdVisualizationProcessor.cpp
+++ b/Engine/Plugins/AI/MassCrowd/Source/MassCrowd/Private/MassCrowdVisualizationProcessor.cpp
@@ -46,9 +46,6 @@ void UMassCrowdVisualizationProcessor::Execute(UMassEntitySubsystem& EntitySubsy
 		UpdateVisualization(Context);
 	});
 
-	// @todo, we should use the new translators to do that initialization
-	InitializeVelocity(EntitySubsystem, Context);
-
 	if (UE::MassCrowd::bDebugCrowdVisualType)
 	{
 		check(World);
diff --git a/Engine/Plugins/AI/MassCrowd/Source/MassCrowd/Public/MassCrowdRepresentationProcessor.h b/Engine/Plugins/AI/MassCrowd/Source/MassCrowd/Public/MassCrowdRepresentationProcessor.h
--- a/Engine/Plugins/AI/MassCrowd/Source/MassCrowd/Public/MassCrowdRepresentationProcessor.h
+++ b/Engine/Plugins/AI/MassCrowd/Source/MassCrowd/Public/MassCrowdRepresentationProcessor.h
@@ -60,13 +60,6 @@ protected:
 	 */
 	virtual void TeleportActor(const FTransform& Transform, AActor& Actor, FMassCommandBuffer& CommandBuffer) override;
 
-	/**
-	 * Initializes velocity for each entity when just switch to use movement component, must be called within a ForEachEntityChunk
-	 * @param EntitySubsystem The subsystem in which all entities must be initialized
-	 * @param Context of the execution from the entity sub system
-	 */
-	void InitializeVelocity(UMassEntitySubsystem& EntitySubsystem, FMassExecutionContext& Context);
-
 	/** 
 	 * A dedicated query for processing entities owning a FDataFragment_CharacterMovementComponentWrapper
 	 */
diff --git a/Engine/Plugins/AI/MassCrowd/Source/MassCrowd/Public/MassCrowdServerRepresentationProcessor.h b/Engine/Plugins/AI/MassCrowd/Source/MassCrowd/Public/MassCrowdServerRepresentationProcessor.h
--- a/Engine/Plugins/AI/MassCrowd/Source/MassCrowd/Public/MassCrowdServerRepresentationProcessor.h
+++ b/Engine/Plugins/AI/MassCrowd/Source/MassCrowd/Public/MassCrowdServerRepresentationProcessor.h
@@ -17,13 +17,4 @@ class MASSCROWD_API UMassCrowdServerRepresentationProcessor : public UMassCrowdR
 
 public:
 	UMassCrowdServerRepresentationProcessor();
-
-protected:
-
-	/**
-	 * Execution method for this processor
-	 * @param EntitySubsystem is the system to execute the lambdas on each entity chunk
-	 * @param Context is the execution context to be passed when executing the lambdas
-	 */
-	virtual void Execute(UMassEntitySubsystem& EntitySubsystem, FMassExecutionContext& Context) override;
 };
\ No newline at end of file
