Better comment on a configuration parameter.

#rnx

diff --git a/Engine/Plugins/AI/MassCrowd/Source/MassCrowd/Private/MassCrowdVisualizationTrait.cpp b/Engine/Plugins/AI/MassCrowd/Source/MassCrowd/Private/MassCrowdVisualizationTrait.cpp
--- a/Engine/Plugins/AI/MassCrowd/Source/MassCrowd/Private/MassCrowdVisualizationTrait.cpp
+++ b/Engine/Plugins/AI/MassCrowd/Source/MassCrowd/Private/MassCrowdVisualizationTrait.cpp
@@ -14,6 +14,7 @@ UMassCrowdVisualizationTrait::UMassCrowdVisualizationTrait()
 	Config.LODRepresentation[EMassLOD::Medium] = ERepresentationType::LowResSpawnedActor;
 	Config.LODRepresentation[EMassLOD::Low] = ERepresentationType::StaticMeshInstance;
 	Config.LODRepresentation[EMassLOD::Off] = ERepresentationType::None;
+	// Set bKeepLowResActor to true as a spawning optimization, this will keep the low-res actor if available while showing the static mesh instance
 	Config.bKeepLowResActors = true;
 	Config.bKeepActorExtraFrame = true;
 	Config.bSpreadFirstVisualizationUpdate = false;
