Mass Crowd: Fixed crowd obstacle movement logic

#rnx

diff --git a/Engine/Plugins/AI/MassCrowd/Source/MassCrowd/Private/MassCrowdNavigationProcessor.cpp b/Engine/Plugins/AI/MassCrowd/Source/MassCrowd/Private/MassCrowdNavigationProcessor.cpp
--- a/Engine/Plugins/AI/MassCrowd/Source/MassCrowd/Private/MassCrowdNavigationProcessor.cpp
+++ b/Engine/Plugins/AI/MassCrowd/Source/MassCrowd/Private/MassCrowdNavigationProcessor.cpp
@@ -5,9 +5,8 @@
 #include "MassAIBehaviorTypes.h"
 #include "MassCrowdSubsystem.h"
 #include "MassCrowdFragments.h"
+#include "MassMovementFragments.h"
 #include "MassCrowdSettings.h"
-#include "MassNavigationFragments.h"
-#include "ZoneGraphSubsystem.h"
 #include "ZoneGraphAnnotationSubsystem.h"
 #include "MassZoneGraphNavigationFragments.h"
 #include "Annotations/ZoneGraphDisturbanceAnnotation.h"
@@ -34,7 +33,7 @@ void UMassCrowdLaneTrackingSignalProcessor::Initialize(UObject& Owner)
 	Super::Initialize(Owner);
 	
 	MassCrowdSubsystem = UWorld::GetSubsystem<UMassCrowdSubsystem>(Owner.GetWorld());
-	checkf(MassCrowdSubsystem != nullptr, TEXT("UMassCrowdSubsystem is mandatory when using MassCrowd processor."));
+	checkf(MassCrowdSubsystem != nullptr, TEXT("UMassCrowdSubsystem is mandatory when using this processor."));
 
 	SubscribeToSignal(UE::Mass::Signals::CurrentLaneChanged);
 }
@@ -79,7 +78,7 @@ void UMassCrowdLaneTrackingDestructor::Initialize(UObject& Owner)
 {
 	Super::Initialize(Owner);
 	MassCrowdSubsystem = UWorld::GetSubsystem<UMassCrowdSubsystem>(Owner.GetWorld());
-	checkf(MassCrowdSubsystem != nullptr, TEXT("UMassCrowdSubsystem is mandatory when using MassCrowd processor."));
+	checkf(MassCrowdSubsystem != nullptr, TEXT("UMassCrowdSubsystem is mandatory when using this processor."));
 }
 
 void UMassCrowdLaneTrackingDestructor::ConfigureQueries()
@@ -122,18 +121,17 @@ void UMassCrowdDynamicObstacleProcessor::Initialize(UObject& Owner)
 {
 	Super::Initialize(Owner);
 
-	ZoneGraphSubsystem = UWorld::GetSubsystem<UZoneGraphSubsystem>(Owner.GetWorld());
 	ZoneGraphAnnotationSubsystem = UWorld::GetSubsystem<UZoneGraphAnnotationSubsystem>(Owner.GetWorld());
-	
-	CrowdSettings = GetDefault<UMassCrowdSettings>();
-	checkf(CrowdSettings, TEXT("Settings default object is always expected to be valid"));
+	checkf(ZoneGraphAnnotationSubsystem != nullptr, TEXT("UZoneGraphAnnotationSubsystem is mandatory when using this processor."));
 }
 
 void UMassCrowdDynamicObstacleProcessor::ConfigureQueries()
 {
 	EntityQuery_Conditional.AddRequirement<FTransformFragment>(EMassFragmentAccess::ReadOnly);
+	EntityQuery_Conditional.AddRequirement<FMassVelocityFragment>(EMassFragmentAccess::ReadOnly, EMassFragmentPresence::Optional);
 	EntityQuery_Conditional.AddRequirement<FAgentRadiusFragment>(EMassFragmentAccess::ReadOnly);
 	EntityQuery_Conditional.AddRequirement<FMassCrowdObstacleFragment>(EMassFragmentAccess::ReadWrite);
+	EntityQuery_Conditional.AddRequirement<FMassSimulationVariableTickFragment>(EMassFragmentAccess::ReadOnly, EMassFragmentPresence::Optional);
 	EntityQuery_Conditional.AddChunkRequirement<FMassSimulationVariableTickChunkFragment>(EMassFragmentAccess::ReadOnly, EMassFragmentPresence::Optional);
 	EntityQuery_Conditional.SetChunkFilter(&FMassSimulationVariableTickChunkFragment::ShouldTickChunkThisFrame);
 }
@@ -141,46 +139,68 @@ void UMassCrowdDynamicObstacleProcessor::ConfigureQueries()
 void UMassCrowdDynamicObstacleProcessor::Execute(UMassEntitySubsystem& EntitySubsystem, FMassExecutionContext& Context)
 {
 	UWorld* World = EntitySubsystem.GetWorld();
+	const UMassCrowdSettings* CrowdSettings = GetDefault<UMassCrowdSettings>();
+	checkf(CrowdSettings, TEXT("Settings default object is always expected to be valid"));
 
-	EntityQuery_Conditional.ForEachEntityChunk(EntitySubsystem, Context, [this, World](FMassExecutionContext& Context)
+	EntityQuery_Conditional.ForEachEntityChunk(EntitySubsystem, Context, [this, World, CrowdSettings](FMassExecutionContext& Context)
 	{
 		const int32 NumEntities = Context.GetNumEntities();
 
 		const TConstArrayView<FTransformFragment> LocationList = Context.GetFragmentView<FTransformFragment>();
+		const TConstArrayView<FMassVelocityFragment> VelocityList = Context.GetFragmentView<FMassVelocityFragment>();
 		const TConstArrayView<FAgentRadiusFragment> RadiusList = Context.GetFragmentView<FAgentRadiusFragment>();
 		const TArrayView<FMassCrowdObstacleFragment> ObstacleDataList = Context.GetMutableFragmentView<FMassCrowdObstacleFragment>();
+		const TConstArrayView<FMassSimulationVariableTickFragment> SimVariableTickList = Context.GetFragmentView<FMassSimulationVariableTickFragment>();
 
-		const float CurrentTime = World->GetTimeSeconds();
-
-		const FDateTime Now = FDateTime::UtcNow();
+		const bool bHasVelocity = VelocityList.Num() > 0;
+		const bool bHasVariableTick = (SimVariableTickList.Num() > 0);
+		const float WorldDeltaTime = Context.GetDeltaTimeSeconds();
 
+		const float ObstacleMovingDistanceTolerance = CrowdSettings->ObstacleMovingDistanceTolerance;
+		const float ObstacleStoppingSpeedTolerance = CrowdSettings->ObstacleStoppingSpeedTolerance;
+		const float ObstacleTimeToStop = CrowdSettings->ObstacleTimeToStop;
+		const float ObstacleEffectRadius = CrowdSettings->ObstacleEffectRadius;
+		
 		for (int32 EntityIndex = 0; EntityIndex < NumEntities; ++EntityIndex)
 		{
 			// @todo: limit update frequency, this does not need to occur every frame
 			const FVector Position = LocationList[EntityIndex].GetTransform().GetLocation();
 			const float Radius = RadiusList[EntityIndex].Radius;
 			FMassCrowdObstacleFragment& Obstacle = ObstacleDataList[EntityIndex];
+			const float DeltaTime = FMath::Max(KINDA_SMALL_NUMBER, bHasVariableTick ? SimVariableTickList[EntityIndex].DeltaTime : WorldDeltaTime);
 
-			UE_VLOG_LOCATION(this, LogMassNavigationObstacle, Display, Position, Radius, Obstacle.bHasStopped ? FColor::Red : FColor::Green, TEXT(""));
+			UE_VLOG_LOCATION(this, LogMassNavigationObstacle, Display, Position, Radius, Obstacle.bIsMoving ? FColor::Green : FColor::Red, TEXT(""));
 
-			if ((Position - Obstacle.LastPosition).SquaredLength() < FMath::Square(DistanceBuffer))
+			if (Obstacle.bIsMoving)
 			{
-				const float TimeElapsed = CurrentTime - Obstacle.LastMovedTimeStamp;
-				if (TimeElapsed > DelayBeforeStopNotification && !Obstacle.bHasStopped)
-				{
-					// The obstacle hasn't moved for a while.
-					Obstacle.bHasStopped = true;
+				// Calculate current speed based on velocity or last known position.
+				const float CurrentSpeed = bHasVelocity ? VelocityList[EntityIndex].Value.Length() : (FVector::Dist(Position, Obstacle.LastPosition) / DeltaTime);
 
-					ensureMsgf(Obstacle.LaneObstacleID.IsValid() == false, TEXT("Obstacle should not have been set."));
+				// Update position while moving, the stop logic will use the last position while check if the obstacles moves again.
+				Obstacle.LastPosition = Position;
 
+				// Keep track how long the obstacle has been almost stationary.
+				if (CurrentSpeed < ObstacleStoppingSpeedTolerance)
+				{
+					Obstacle.TimeSinceStopped += DeltaTime;
+				}
+				else
+				{
+					Obstacle.TimeSinceStopped = 0.0f;
+				}
+				
+				// If the obstacle has been almost stationary for a while, mark it as obstacle.
+				if (Obstacle.TimeSinceStopped > ObstacleTimeToStop)
+				{
+					ensureMsgf(Obstacle.LaneObstacleID.IsValid() == false, TEXT("Obstacle should not have been set."));
+						
+					Obstacle.bIsMoving = false;
 					Obstacle.LaneObstacleID = FMassLaneObstacleID::GetNextUniqueID();
 
-					constexpr float EffectRadius = 1000.f;
-
 					// Add an obstacle disturbance.
 					FZoneGraphObstacleDisturbanceArea Disturbance;
 					Disturbance.Position = Obstacle.LastPosition;
-					Disturbance.Radius = EffectRadius;
+					Disturbance.Radius = ObstacleEffectRadius;
 					Disturbance.ObstacleRadius = Radius;
 					Disturbance.ObstacleID = Obstacle.LaneObstacleID;
 					Disturbance.Action = EZoneGraphObstacleDisturbanceAreaAction::Add;
@@ -189,26 +209,110 @@ void UMassCrowdDynamicObstacleProcessor::Execute(UMassEntitySubsystem& EntitySub
 			}
 			else
 			{
-				// Update position and time stamp
-				Obstacle.LastPosition = Position;
-				Obstacle.LastMovedTimeStamp = CurrentTime;
+				Obstacle.TimeSinceStopped += DeltaTime;
 
-				// If the obstacle had stopped, signal the move.
-				if (Obstacle.bHasStopped)
+				// If the obstacle moves outside movement tolerance, mark it as moving, and remove it as obstacle.
+				if (FVector::Dist(Position, Obstacle.LastPosition) > ObstacleMovingDistanceTolerance)				
 				{
-					Obstacle.bHasStopped = false;
+					ensureMsgf(Obstacle.LaneObstacleID.IsValid(), TEXT("Obstacle should have been set."));
 
-					if (ensureMsgf(Obstacle.LaneObstacleID.IsValid(), TEXT("Obstacle should have been set.")))
-					{
-						FZoneGraphObstacleDisturbanceArea Disturbance;
-						Disturbance.ObstacleID = Obstacle.LaneObstacleID;
-						Disturbance.Action = EZoneGraphObstacleDisturbanceAreaAction::Remove;
-						ZoneGraphAnnotationSubsystem->SendEvent(Disturbance);
+					FZoneGraphObstacleDisturbanceArea Disturbance;
+					Disturbance.ObstacleID = Obstacle.LaneObstacleID;
+					Disturbance.Action = EZoneGraphObstacleDisturbanceAreaAction::Remove;
+					ZoneGraphAnnotationSubsystem->SendEvent(Disturbance);
 
-						Obstacle.LaneObstacleID = FMassLaneObstacleID();
-					}
+					Obstacle.bIsMoving = true;
+					Obstacle.TimeSinceStopped = 0.0f;
+					Obstacle.LaneObstacleID = FMassLaneObstacleID();
 				}
 			}
 		}
 	});
 }
+
+//----------------------------------------------------------------------//
+// UMassCrowdDynamicObstacleInitializer
+//----------------------------------------------------------------------//
+UMassCrowdDynamicObstacleInitializer::UMassCrowdDynamicObstacleInitializer()
+{
+	ExecutionFlags = (int32)(EProcessorExecutionFlags::Standalone | EProcessorExecutionFlags::Server);
+	ObservedType = FMassCrowdObstacleFragment::StaticStruct();
+	Operation = EMassObservedOperation::Add;
+}
+
+void UMassCrowdDynamicObstacleInitializer::ConfigureQueries()
+{
+	EntityQuery.AddRequirement<FTransformFragment>(EMassFragmentAccess::ReadOnly);
+	EntityQuery.AddRequirement<FMassCrowdObstacleFragment>(EMassFragmentAccess::ReadWrite);
+}
+
+void UMassCrowdDynamicObstacleInitializer::Execute(UMassEntitySubsystem& EntitySubsystem, FMassExecutionContext& Context)
+{
+	UWorld* World = EntitySubsystem.GetWorld();
+	
+	EntityQuery.ForEachEntityChunk(EntitySubsystem, Context, [World](FMassExecutionContext& Context)
+	{
+		const int32 NumEntities = Context.GetNumEntities();
+		const TConstArrayView<FTransformFragment> LocationList = Context.GetFragmentView<FTransformFragment>();
+		const TArrayView<FMassCrowdObstacleFragment> ObstacleDataList = Context.GetMutableFragmentView<FMassCrowdObstacleFragment>();
+
+		for (int32 EntityIndex = 0; EntityIndex < NumEntities; ++EntityIndex)
+		{
+			const FVector Position = LocationList[EntityIndex].GetTransform().GetLocation();
+			FMassCrowdObstacleFragment& Obstacle = ObstacleDataList[EntityIndex];
+
+			Obstacle.LastPosition = Position;
+			Obstacle.TimeSinceStopped = 0.0f;
+			Obstacle.bIsMoving = true;
+		}
+	});
+}
+
+
+//----------------------------------------------------------------------//
+// UMassCrowdDynamicObstacleDeinitializer
+//----------------------------------------------------------------------//
+UMassCrowdDynamicObstacleDeinitializer::UMassCrowdDynamicObstacleDeinitializer()
+{
+	ExecutionFlags = (int32)(EProcessorExecutionFlags::Standalone | EProcessorExecutionFlags::Server);
+	ObservedType = FMassCrowdObstacleFragment::StaticStruct();
+	Operation = EMassObservedOperation::Remove;
+}
+
+void UMassCrowdDynamicObstacleDeinitializer::Initialize(UObject& Owner)
+{
+	Super::Initialize(Owner);
+
+	ZoneGraphAnnotationSubsystem = UWorld::GetSubsystem<UZoneGraphAnnotationSubsystem>(Owner.GetWorld());
+	checkf(ZoneGraphAnnotationSubsystem != nullptr, TEXT("UZoneGraphAnnotationSubsystem is mandatory when using this processor."));
+}
+
+void UMassCrowdDynamicObstacleDeinitializer::ConfigureQueries()
+{
+	EntityQuery.AddRequirement<FMassCrowdObstacleFragment>(EMassFragmentAccess::ReadWrite);
+}
+
+void UMassCrowdDynamicObstacleDeinitializer::Execute(UMassEntitySubsystem& EntitySubsystem, FMassExecutionContext& Context)
+{
+	EntityQuery.ForEachEntityChunk(EntitySubsystem, Context, [this](FMassExecutionContext& Context)
+	{
+		const int32 NumEntities = Context.GetNumEntities();
+		const TArrayView<FMassCrowdObstacleFragment> ObstacleDataList = Context.GetMutableFragmentView<FMassCrowdObstacleFragment>();
+
+		for (int32 EntityIndex = 0; EntityIndex < NumEntities; ++EntityIndex)
+		{
+			FMassCrowdObstacleFragment& Obstacle = ObstacleDataList[EntityIndex];
+
+			if (Obstacle.LaneObstacleID.IsValid())
+			{
+				FZoneGraphObstacleDisturbanceArea Disturbance;
+				Disturbance.ObstacleID = Obstacle.LaneObstacleID;
+				Disturbance.Action = EZoneGraphObstacleDisturbanceAreaAction::Remove;
+				ZoneGraphAnnotationSubsystem->SendEvent(Disturbance);
+
+				// Reset obstacle
+				Obstacle = FMassCrowdObstacleFragment();
+			}
+		}
+	});
+}
\ No newline at end of file
diff --git a/Engine/Plugins/AI/MassCrowd/Source/MassCrowd/Public/MassCrowdFragments.h b/Engine/Plugins/AI/MassCrowd/Source/MassCrowd/Public/MassCrowdFragments.h
--- a/Engine/Plugins/AI/MassCrowd/Source/MassCrowd/Public/MassCrowdFragments.h
+++ b/Engine/Plugins/AI/MassCrowd/Source/MassCrowd/Public/MassCrowdFragments.h
@@ -36,12 +36,12 @@ struct MASSCROWD_API FMassCrowdObstacleFragment : public FMassFragment
 	/** Obstacle ID reported to the obstruction annotation. */
 	FMassLaneObstacleID LaneObstacleID;
 
-	/** Position of the dynamic obstacle when it last moved. */
+	/** Position of the obstacle when it last moved. */
 	FVector LastPosition = FVector::ZeroVector;
 
-	/** Time stamp when that obstacle stopped moving. */
-	float LastMovedTimeStamp = 0.0f;
+	/** Time since the obstacle has not moved based on speed tolerance. */
+	float TimeSinceStopped = 0.0f;
 
-	/** Has this dynamic obstacle stopped moving. */
-	bool bHasStopped = true;
+	/** True of the current obstacle state is moving. */
+	bool bIsMoving = true;
 };
diff --git a/Engine/Plugins/AI/MassCrowd/Source/MassCrowd/Public/MassCrowdNavigationProcessor.h b/Engine/Plugins/AI/MassCrowd/Source/MassCrowd/Public/MassCrowdNavigationProcessor.h
--- a/Engine/Plugins/AI/MassCrowd/Source/MassCrowd/Public/MassCrowdNavigationProcessor.h
+++ b/Engine/Plugins/AI/MassCrowd/Source/MassCrowd/Public/MassCrowdNavigationProcessor.h
@@ -7,8 +7,6 @@
 #include "MassCrowdFragments.h"
 #include "MassCrowdNavigationProcessor.generated.h"
 
-class UZoneGraphSubsystem;
-class UMassCrowdSettings;
 class UZoneGraphAnnotationSubsystem;
 class UMassCrowdSubsystem;
 
@@ -49,6 +47,7 @@ protected:
 	FMassEntityQuery EntityQuery;
 };
 
+
 UCLASS()
 class MASSCROWD_API UMassCrowdDynamicObstacleProcessor : public UMassProcessor
 {
@@ -65,22 +64,44 @@ protected:
 	void OnStop(FMassCrowdObstacleFragment& OutObstacle, const float Radius);
 	void OnMove(FMassCrowdObstacleFragment& OutObstacle);
 
-	/** Delay before sending the stop notification once the entity has stop moving. */
-	UPROPERTY(Category = Settings, EditAnywhere, meta = (ClampMin = "0"))
-	float DelayBeforeStopNotification = 0.3f;
-
-	/** Distance within which the obstacle is considered not moving. */
-	UPROPERTY(Category = Settings, EditAnywhere, meta = (ClampMin = "0"))
-	float DistanceBuffer = 10.f;
-
 	FMassEntityQuery EntityQuery_Conditional;
 
-	UPROPERTY(Transient)
-	UZoneGraphSubsystem* ZoneGraphSubsystem;
-	
 	UPROPERTY(Transient)
 	UZoneGraphAnnotationSubsystem* ZoneGraphAnnotationSubsystem;
+};
+
+
+UCLASS()
+class MASSCROWD_API UMassCrowdDynamicObstacleInitializer : public UMassObserverProcessor
+{
+	GENERATED_BODY()
+
+public:
+	UMassCrowdDynamicObstacleInitializer();
+
+protected:
+	virtual void ConfigureQueries() override;
+	virtual void Execute(UMassEntitySubsystem& EntitySubsystem, FMassExecutionContext& Context) override;
+
+	FMassEntityQuery EntityQuery;
+};
+
+
+UCLASS()
+class MASSCROWD_API UMassCrowdDynamicObstacleDeinitializer : public UMassObserverProcessor
+{
+	GENERATED_BODY()
+
+public:
+	UMassCrowdDynamicObstacleDeinitializer();
+
+protected:
+	virtual void Initialize(UObject& Owner) override;
+	virtual void ConfigureQueries() override;
+	virtual void Execute(UMassEntitySubsystem& EntitySubsystem, FMassExecutionContext& Context) override;
 
 	UPROPERTY(Transient)
-	const UMassCrowdSettings* CrowdSettings;
+	UZoneGraphAnnotationSubsystem* ZoneGraphAnnotationSubsystem;
+
+	FMassEntityQuery EntityQuery;
 };
diff --git a/Engine/Plugins/AI/MassCrowd/Source/MassCrowd/Public/MassCrowdSettings.h b/Engine/Plugins/AI/MassCrowd/Source/MassCrowd/Public/MassCrowdSettings.h
--- a/Engine/Plugins/AI/MassCrowd/Source/MassCrowd/Public/MassCrowdSettings.h
+++ b/Engine/Plugins/AI/MassCrowd/Source/MassCrowd/Public/MassCrowdSettings.h
@@ -115,6 +115,22 @@ public:
 	UPROPERTY(EditAnywhere, config, Category = Options)
 	bool bDisplayDensities = true;
 
+	/** An obstacle is considered being stopped when it's speed is less than the tolerance. */
+	UPROPERTY(EditAnywhere, config, Category = Obstacles, meta = (ClampMin = "0.0", UIMin = "0.0", ForceUnits = "cm/s"))
+	float ObstacleStoppingSpeedTolerance = 5.0f;
+
+	/** An obstacle is considered moving when it has moved this much after being stationary. */
+	UPROPERTY(EditAnywhere, config, Category = Obstacles, meta = (ClampMin = "0.0", UIMin = "0.0", ForceUnits = "cm"))
+	float ObstacleMovingDistanceTolerance = 10.0f;
+
+	/** The time an obstacle needs to be not moving before it is reported as stopped.*/
+	UPROPERTY(EditAnywhere, config, Category = Obstacles, meta = (ClampMin = "0.0", UIMin = "0.0", ForceUnits = "s"))
+	float ObstacleTimeToStop = 0.3f;
+
+	/** The radius an obstacle has effects on navigation.*/
+	UPROPERTY(EditAnywhere, config, Category = Obstacles, meta = (ClampMin = "0.0", UIMin = "0.0", ForceUnits = "cm"))
+	float ObstacleEffectRadius = 1000.f;
+
 protected:
 
 #if WITH_EDITOR
