Moved Visualization LOD Processor parameters into shared fragment

#rnx

diff --git a/Engine/Plugins/AI/MassCrowd/Source/MassCrowd/Private/MassCrowdVisualizationLODProcessor.cpp b/Engine/Plugins/AI/MassCrowd/Source/MassCrowd/Private/MassCrowdVisualizationLODProcessor.cpp
--- a/Engine/Plugins/AI/MassCrowd/Source/MassCrowd/Private/MassCrowdVisualizationLODProcessor.cpp
+++ b/Engine/Plugins/AI/MassCrowd/Source/MassCrowd/Private/MassCrowdVisualizationLODProcessor.cpp
@@ -19,22 +19,22 @@ namespace UE::MassCrowd
 		FAutoConsoleVariableRef(TEXT("ai.debug.ShowISMUnderSpecifiedRange"), bDebugShowISMUnderSpecifiedRange, TEXT("Show ISM under a specified range (meters)"), ECVF_Cheat)
 	};
 
-#if WITH_EDITOR
-	int32 DebugCrowdMaxCountHigh = -1;
-	int32 DebugCrowdMaxCountMedium = -1;
-	int32 DebugCrowdMaxCountLow = -1;
-
-	FAutoConsoleVariableRef EditorConsoleVariables[] =
-	{
-		FAutoConsoleVariableRef(TEXT("ai.debug.CrowdMaxCountHigh"), DebugCrowdMaxCountHigh, TEXT("Forced MASS crowd max count high"), ECVF_Cheat),
-		FAutoConsoleVariableRef(TEXT("ai.debug.CrowdMaxCountMedium"), DebugCrowdMaxCountMedium, TEXT("Forced MASS crowd max count medium"), ECVF_Cheat),
-		FAutoConsoleVariableRef(TEXT("ai.debug.CrowdMaxCountLow"), DebugCrowdMaxCountLow, TEXT("Forced MASS crowd max count low"), ECVF_Cheat)
-	};
-
-	int32 LastDebugCrowdMaxCountHigh = -1;
-	int32 LastDebugCrowdMaxCountMedium = -1;
-	int32 LastDebugCrowdMaxCountLow = -1;
-#endif // WITH_EDITOR
+// #if WITH_EDITOR
+// 	int32 DebugCrowdMaxCountHigh = -1;
+// 	int32 DebugCrowdMaxCountMedium = -1;
+// 	int32 DebugCrowdMaxCountLow = -1;
+// 
+// 	FAutoConsoleVariableRef EditorConsoleVariables[] =
+// 	{
+// 		FAutoConsoleVariableRef(TEXT("ai.debug.CrowdMaxCountHigh"), DebugCrowdMaxCountHigh, TEXT("Forced MASS crowd max count high"), ECVF_Cheat),
+// 		FAutoConsoleVariableRef(TEXT("ai.debug.CrowdMaxCountMedium"), DebugCrowdMaxCountMedium, TEXT("Forced MASS crowd max count medium"), ECVF_Cheat),
+// 		FAutoConsoleVariableRef(TEXT("ai.debug.CrowdMaxCountLow"), DebugCrowdMaxCountLow, TEXT("Forced MASS crowd max count low"), ECVF_Cheat)
+// 	};
+// 
+// 	int32 LastDebugCrowdMaxCountHigh = -1;
+// 	int32 LastDebugCrowdMaxCountMedium = -1;
+// 	int32 LastDebugCrowdMaxCountLow = -1;
+// #endif // WITH_EDITOR
 } // UE::MassCrowd
 
 UMassCrowdVisualizationLODProcessor::UMassCrowdVisualizationLODProcessor()
@@ -46,28 +46,21 @@ UMassCrowdVisualizationLODProcessor::UMassCrowdVisualizationLODProcessor()
 	ExecutionOrder.ExecuteInGroup = UE::Mass::ProcessorGroupNames::LOD;
 	ExecutionOrder.ExecuteAfter.Add(UE::Mass::ProcessorGroupNames::LODCollector);
 
-#if WITH_EDITOR
-	UE::MassCrowd::LastDebugCrowdMaxCountHigh = -1;
-	UE::MassCrowd::LastDebugCrowdMaxCountMedium = -1;
-	UE::MassCrowd::LastDebugCrowdMaxCountLow = -1;
-#endif // WITH_EDITOR
+// #if WITH_EDITOR
+// 	UE::MassCrowd::LastDebugCrowdMaxCountHigh = -1;
+// 	UE::MassCrowd::LastDebugCrowdMaxCountMedium = -1;
+// 	UE::MassCrowd::LastDebugCrowdMaxCountLow = -1;
+// #endif // WITH_EDITOR
 }
 
 void UMassCrowdVisualizationLODProcessor::ConfigureQueries()
 {
-	// @todo remove, no need for this anymore since we have the common LOD collector
-	// Do not call super as we do have our own LODInfo fragment, so need to duplicate
+	Super::ConfigureQueries();
 
 	CloseEntityQuery.AddTagRequirement<FMassCrowdTag>(EMassFragmentPresence::All);
-	CloseEntityQuery.AddRequirement<FMassViewerInfoFragment>(EMassFragmentAccess::ReadOnly);
-	CloseEntityQuery.AddRequirement<FMassRepresentationFragment>(EMassFragmentAccess::ReadOnly);
-	CloseEntityQuery.AddRequirement<FMassRepresentationLODFragment>(EMassFragmentAccess::ReadWrite);
-	CloseEntityQuery.AddRequirement<FTransformFragment>(EMassFragmentAccess::ReadOnly);
-
-	FarEntityQuery = CloseEntityQuery;
-	CloseEntityQuery.AddTagRequirement<FMassVisibilityCulledByDistanceTag>(EMassFragmentPresence::None);
-	FarEntityQuery.AddTagRequirement<FMassVisibilityCulledByDistanceTag>(EMassFragmentPresence::All);
-	FarEntityQuery.AddChunkRequirement<FMassVisualizationChunkFragment>(EMassFragmentAccess::ReadOnly);
+	CloseEntityAdjustDistanceQuery.AddTagRequirement<FMassCrowdTag>(EMassFragmentPresence::All);
+	FarEntityQuery.AddTagRequirement<FMassCrowdTag>(EMassFragmentPresence::All);
+	DebugEntityQuery.AddTagRequirement<FMassCrowdTag>(EMassFragmentPresence::All);
 }
 
 void UMassCrowdVisualizationLODProcessor::Execute(UMassEntitySubsystem& EntitySubsystem, FMassExecutionContext& Context)
@@ -76,28 +69,26 @@ void UMassCrowdVisualizationLODProcessor::Execute(UMassEntitySubsystem& EntitySu
 
 	TRACE_CPUPROFILER_EVENT_SCOPE(TEXT("CrowdVisualizationLOD"))
 
-	Super::ExecuteInternal<FMassViewerInfoFragment>(EntitySubsystem, Context);
+	Super::Execute(EntitySubsystem, Context);
 	
 	if (UE::MassCrowd::bDebugCrowdVisualizationLOD)
 	{
 		TRACE_CPUPROFILER_EVENT_SCOPE(TEXT("DebugDisplayLOD"))
 
-		auto DebugDisplayLOD = [this](FMassExecutionContext& Context)
+		DebugEntityQuery.ForEachEntityChunk(EntitySubsystem, Context, [this](FMassExecutionContext& Context)
 		{
+			FMassVisualizationLODSharedFragment& LODSharedFragment = Context.GetMutableSharedFragment<FMassVisualizationLODSharedFragment>();
 			const TConstArrayView<FTransformFragment> LocationList = Context.GetFragmentView<FTransformFragment>();
 			const TConstArrayView<FMassRepresentationLODFragment> VisualizationLODList = Context.GetFragmentView<FMassRepresentationLODFragment>();
-			LODCalculator.DebugDisplayLOD(Context, VisualizationLODList, LocationList, World);
-		};
-
-		CloseEntityQuery.ForEachEntityChunk(EntitySubsystem, Context, DebugDisplayLOD);
-		FarEntityQuery.ForEachEntityChunk(EntitySubsystem, Context, DebugDisplayLOD);
+			LODSharedFragment.LODCalculator.DebugDisplayLOD(Context, VisualizationLODList, LocationList, World);
+		});
 	}
 
 	if (UE::MassCrowd::bDebugShowISMUnderSpecifiedRange > 0)
 	{
 		TRACE_CPUPROFILER_EVENT_SCOPE(TEXT("ShowISMUnderSpecifiedRange"))
 
-		auto ShowISMUnderSpecifiedRange = [this](const FMassExecutionContext& Context)
+		DebugEntityQuery.ForEachEntityChunk(EntitySubsystem, Context, [this](const FMassExecutionContext& Context)
 		{
 			const TConstArrayView<FTransformFragment> LocationList = Context.GetFragmentView<FTransformFragment>();
 			const TConstArrayView<FMassRepresentationFragment> RepresentationFragmentList = Context.GetFragmentView<FMassRepresentationFragment>();
@@ -114,50 +105,49 @@ void UMassCrowdVisualizationLODProcessor::Execute(UMassEntitySubsystem& EntitySu
 					DrawDebugSolidBox(World, EntityLocation.GetTransform().GetLocation() + FVector(0.0f, 0.0f, 150.0f), FVector(50.0f), FColor::Red);
 				}
 			}
-		};
-		CloseEntityQuery.ForEachEntityChunk(EntitySubsystem, Context, ShowISMUnderSpecifiedRange);
-		FarEntityQuery.ForEachEntityChunk(EntitySubsystem, Context, ShowISMUnderSpecifiedRange);
+		});
 	}
 
-#if WITH_EDITOR
-	if (UE::MassCrowd::LastDebugCrowdMaxCountHigh != UE::MassCrowd::DebugCrowdMaxCountHigh)
-	{
-		if (UE::MassCrowd::DebugCrowdMaxCountHigh >= 0)
-		{
-			LODMaxCount[0] = UE::MassCrowd::DebugCrowdMaxCountHigh;
-		}
-		else
-		{
-			LODMaxCount[0] = GetClass()->GetDefaultObject<UMassCrowdVisualizationLODProcessor>()->LODMaxCount[0];
-		}
-		UE::MassCrowd::LastDebugCrowdMaxCountHigh = UE::MassCrowd::DebugCrowdMaxCountHigh;
-		LODCalculator.Initialize(BaseLODDistance, BufferHysteresisOnDistancePercentage / 100.f, LODMaxCount, nullptr, DistanceToFrustum, DistanceToFrustumHysteresis, VisibleLODDistance);
-	}
-	if (UE::MassCrowd::LastDebugCrowdMaxCountMedium != UE::MassCrowd::DebugCrowdMaxCountMedium)
-	{
-		if (UE::MassCrowd::DebugCrowdMaxCountMedium >= 0)
-		{
-			LODMaxCount[1] = UE::MassCrowd::DebugCrowdMaxCountMedium;
-		}
-		else
-		{
-			LODMaxCount[1] = GetClass()->GetDefaultObject<UMassCrowdVisualizationLODProcessor>()->LODMaxCount[1];
-		}
-		UE::MassCrowd::LastDebugCrowdMaxCountMedium = UE::MassCrowd::DebugCrowdMaxCountMedium;
-		LODCalculator.Initialize(BaseLODDistance, BufferHysteresisOnDistancePercentage / 100.f, LODMaxCount, nullptr, DistanceToFrustum, DistanceToFrustumHysteresis, VisibleLODDistance);
-	}
-	if (UE::MassCrowd::LastDebugCrowdMaxCountLow != UE::MassCrowd::DebugCrowdMaxCountLow)
-	{
-		if (UE::MassCrowd::DebugCrowdMaxCountLow >= 0)
-		{
-			LODMaxCount[2] = UE::MassCrowd::DebugCrowdMaxCountLow;
-		}
-		else
-		{
-			LODMaxCount[2] = GetClass()->GetDefaultObject<UMassCrowdVisualizationLODProcessor>()->LODMaxCount[2];
-		}
-		UE::MassCrowd::LastDebugCrowdMaxCountLow = UE::MassCrowd::DebugCrowdMaxCountLow;
-		LODCalculator.Initialize(BaseLODDistance, BufferHysteresisOnDistancePercentage / 100.f, LODMaxCount, nullptr, DistanceToFrustum, DistanceToFrustumHysteresis, VisibleLODDistance);
-	}
-#endif // WITH_EDITOR
+// @todo find a way to do this now with shared fragments
+// #if WITH_EDITOR
+// 	if (UE::MassCrowd::LastDebugCrowdMaxCountHigh != UE::MassCrowd::DebugCrowdMaxCountHigh)
+// 	{
+// 		if (UE::MassCrowd::DebugCrowdMaxCountHigh >= 0)
+// 		{
+// 			LODMaxCount[0] = UE::MassCrowd::DebugCrowdMaxCountHigh;
+// 		}
+// 		else
+// 		{
+// 			LODMaxCount[0] = GetClass()->GetDefaultObject<UMassCrowdVisualizationLODProcessor>()->LODMaxCount[0];
+// 		}
+// 		UE::MassCrowd::LastDebugCrowdMaxCountHigh = UE::MassCrowd::DebugCrowdMaxCountHigh;
+// 		LODCalculator.Initialize(BaseLODDistance, BufferHysteresisOnDistancePercentage / 100.f, LODMaxCount, nullptr, DistanceToFrustum, DistanceToFrustumHysteresis, VisibleLODDistance);
+// 	}
+// 	if (UE::MassCrowd::LastDebugCrowdMaxCountMedium != UE::MassCrowd::DebugCrowdMaxCountMedium)
+// 	{
+// 		if (UE::MassCrowd::DebugCrowdMaxCountMedium >= 0)
+// 		{
+// 			LODMaxCount[1] = UE::MassCrowd::DebugCrowdMaxCountMedium;
+// 		}
+// 		else
+// 		{
+// 			LODMaxCount[1] = GetClass()->GetDefaultObject<UMassCrowdVisualizationLODProcessor>()->LODMaxCount[1];
+// 		}
+// 		UE::MassCrowd::LastDebugCrowdMaxCountMedium = UE::MassCrowd::DebugCrowdMaxCountMedium;
+// 		LODCalculator.Initialize(BaseLODDistance, BufferHysteresisOnDistancePercentage / 100.f, LODMaxCount, nullptr, DistanceToFrustum, DistanceToFrustumHysteresis, VisibleLODDistance);
+// 	}
+// 	if (UE::MassCrowd::LastDebugCrowdMaxCountLow != UE::MassCrowd::DebugCrowdMaxCountLow)
+// 	{
+// 		if (UE::MassCrowd::DebugCrowdMaxCountLow >= 0)
+// 		{
+// 			LODMaxCount[2] = UE::MassCrowd::DebugCrowdMaxCountLow;
+// 		}
+// 		else
+// 		{
+// 			LODMaxCount[2] = GetClass()->GetDefaultObject<UMassCrowdVisualizationLODProcessor>()->LODMaxCount[2];
+// 		}
+// 		UE::MassCrowd::LastDebugCrowdMaxCountLow = UE::MassCrowd::DebugCrowdMaxCountLow;
+// 		LODCalculator.Initialize(BaseLODDistance, BufferHysteresisOnDistancePercentage / 100.f, LODMaxCount, nullptr, DistanceToFrustum, DistanceToFrustumHysteresis, VisibleLODDistance);
+// 	}
+// #endif // WITH_EDITOR
 }
diff --git a/Engine/Plugins/AI/MassCrowd/Source/MassCrowd/Private/MassCrowdVisualizationTrait.cpp b/Engine/Plugins/AI/MassCrowd/Source/MassCrowd/Private/MassCrowdVisualizationTrait.cpp
--- a/Engine/Plugins/AI/MassCrowd/Source/MassCrowd/Private/MassCrowdVisualizationTrait.cpp
+++ b/Engine/Plugins/AI/MassCrowd/Source/MassCrowd/Private/MassCrowdVisualizationTrait.cpp
@@ -20,4 +20,23 @@ UMassCrowdVisualizationTrait::UMassCrowdVisualizationTrait()
 	Params.bSpreadFirstVisualizationUpdate = false;
 	Params.WorldPartitionGridNameContainingCollision = NAME_None;
 	Params.NotVisibleUpdateRate = 0.5f;
+
+	LODParams.BaseLODDistance[EMassLOD::High] = 0.f;
+	LODParams.BaseLODDistance[EMassLOD::Medium] = 0.f;
+	LODParams.BaseLODDistance[EMassLOD::Low] = 3000.f;
+	LODParams.BaseLODDistance[EMassLOD::Off] = 6000.f;
+
+	LODParams.VisibleLODDistance[EMassLOD::High] = 0.f;
+	LODParams.VisibleLODDistance[EMassLOD::Medium] = 0.f;
+	LODParams.VisibleLODDistance[EMassLOD::Low] = 6000.f;
+	LODParams.VisibleLODDistance[EMassLOD::Off] = 50000.f;
+
+	LODParams.LODMaxCount[EMassLOD::High] = 10;
+	LODParams.LODMaxCount[EMassLOD::Medium] = 20;
+	LODParams.LODMaxCount[EMassLOD::Low] = 500;
+	LODParams.LODMaxCount[EMassLOD::Off] = TNumericLimits<int32>::Max();
+
+	LODParams.BufferHysteresisOnDistancePercentage = 20.0f;
+	LODParams.DistanceToFrustum = 0.0f;
+	LODParams.DistanceToFrustumHysteresis = 0.0f;
 }
diff --git a/Engine/Plugins/Runtime/MassGameplay/Source/MassRepresentation/Private/MassRepresentationFragments.cpp b/Engine/Plugins/Runtime/MassGameplay/Source/MassRepresentation/Private/MassRepresentationFragments.cpp
new file mode 100644
--- /dev/null
+++ b/Engine/Plugins/Runtime/MassGameplay/Source/MassRepresentation/Private/MassRepresentationFragments.cpp
@@ -0,0 +1,8 @@
+// Copyright Epic Games, Inc. All Rights Reserved.
+
+#include "MassRepresentationFragments.h"
+
+FMassVisualizationLODSharedFragment::FMassVisualizationLODSharedFragment(const FMassVisualizationLODParameters& LODParams)
+{
+	LODCalculator.Initialize(LODParams.BaseLODDistance, LODParams.BufferHysteresisOnDistancePercentage / 100.f, LODParams.LODMaxCount, nullptr, LODParams.DistanceToFrustum, LODParams.DistanceToFrustumHysteresis, LODParams.VisibleLODDistance);
+}
\ No newline at end of file
diff --git a/Engine/Plugins/Runtime/MassGameplay/Source/MassRepresentation/Private/MassVisualizationLODProcessor.cpp b/Engine/Plugins/Runtime/MassGameplay/Source/MassRepresentation/Private/MassVisualizationLODProcessor.cpp
--- a/Engine/Plugins/Runtime/MassGameplay/Source/MassRepresentation/Private/MassVisualizationLODProcessor.cpp
+++ b/Engine/Plugins/Runtime/MassGameplay/Source/MassRepresentation/Private/MassVisualizationLODProcessor.cpp
@@ -14,54 +14,100 @@ UMassVisualizationLODProcessor::UMassVisualizationLODProcessor()
 	bAutoRegisterWithProcessingPhases = false;
 
 	ExecutionOrder.ExecuteInGroup = UE::Mass::ProcessorGroupNames::LOD;
-
-	BaseLODDistance[EMassLOD::High] = 0.f;
-	BaseLODDistance[EMassLOD::Medium] = 1000.f;
-	BaseLODDistance[EMassLOD::Low] = 2500.f;
-	BaseLODDistance[EMassLOD::Off] = 10000.f;
-	
-	VisibleLODDistance[EMassLOD::High] = 0.f;
-	VisibleLODDistance[EMassLOD::Medium] = 2000.f;
-	VisibleLODDistance[EMassLOD::Low] = 4000.f;
-	VisibleLODDistance[EMassLOD::Off] = 10000.f;
-	
-	LODMaxCount[EMassLOD::High] = 50;
-	LODMaxCount[EMassLOD::Medium] = 100;
-	LODMaxCount[EMassLOD::Low] = 500;
-	LODMaxCount[EMassLOD::Off] = 0;
+	ExecutionOrder.ExecuteAfter.Add(UE::Mass::ProcessorGroupNames::LODCollector);
 }
 
 void UMassVisualizationLODProcessor::ConfigureQueries()
 {
-	CloseEntityQuery.AddRequirement<FMassViewerInfoFragment>(EMassFragmentAccess::ReadOnly);
-	CloseEntityQuery.AddRequirement<FMassRepresentationLODFragment>(EMassFragmentAccess::ReadWrite);
-	CloseEntityQuery.AddRequirement<FTransformFragment>(EMassFragmentAccess::ReadOnly);
-
-	FarEntityQuery = CloseEntityQuery;
+	FMassEntityQuery BaseQuery;
+	BaseQuery.AddRequirement<FMassViewerInfoFragment>(EMassFragmentAccess::ReadOnly);
+	BaseQuery.AddRequirement<FMassRepresentationLODFragment>(EMassFragmentAccess::ReadWrite);
+	BaseQuery.AddRequirement<FTransformFragment>(EMassFragmentAccess::ReadOnly);
+	BaseQuery.AddConstSharedRequirement<FMassVisualizationLODParameters>();
+	BaseQuery.AddSharedRequirement<FMassVisualizationLODSharedFragment>(EMassFragmentAccess::ReadWrite);
 
+	CloseEntityQuery = BaseQuery;
 	CloseEntityQuery.AddTagRequirement<FMassVisibilityCulledByDistanceTag>(EMassFragmentPresence::None);
 	
+	CloseEntityAdjustDistanceQuery = CloseEntityQuery;
+	CloseEntityAdjustDistanceQuery.SetArchetypeFilter([](const FMassExecutionContext& Context)
+	{
+		const FMassVisualizationLODSharedFragment& LODSharedFragment = Context.GetSharedFragment<FMassVisualizationLODSharedFragment>();
+		return LODSharedFragment.bHasAdjustedDistancesFromCount;
+	});
+
+	FarEntityQuery = BaseQuery;
 	FarEntityQuery.AddTagRequirement<FMassVisibilityCulledByDistanceTag>(EMassFragmentPresence::All);
 	FarEntityQuery.AddChunkRequirement<FMassVisualizationChunkFragment>(EMassFragmentAccess::ReadOnly);
+	FarEntityQuery.SetChunkFilter(&FMassVisualizationChunkFragment::ShouldUpdateVisualizationForChunk);
+
+	DebugEntityQuery = BaseQuery;
 }
 
-void UMassVisualizationLODProcessor::Initialize(UObject& Owner)
+void UMassVisualizationLODProcessor::Execute(UMassEntitySubsystem& EntitySubsystem, FMassExecutionContext& Context)
 {
-	Super::Initialize(Owner);
+	if (bForceOFFLOD)
+	{
+		CloseEntityQuery.ForEachEntityChunk(EntitySubsystem, Context, [this](FMassExecutionContext& Context)
+		{
+			FMassVisualizationLODSharedFragment& LODSharedFragment = Context.GetMutableSharedFragment<FMassVisualizationLODSharedFragment>();
+			TArrayView<FMassRepresentationLODFragment> RepresentationLODList = Context.GetMutableFragmentView<FMassRepresentationLODFragment>();
+			LODSharedFragment.LODCalculator.ForceOffLOD(Context, RepresentationLODList);
+		});
+		return;
+	}
 
-	LODCalculator.Initialize(BaseLODDistance, BufferHysteresisOnDistancePercentage / 100.f, LODMaxCount, nullptr, DistanceToFrustum, DistanceToFrustumHysteresis, VisibleLODDistance);
-}
+	{
+		TRACE_CPUPROFILER_EVENT_SCOPE(PrepareExecution)
+		check(LODSubsystem);
+		const TArray<FViewerInfo>& Viewers = LODSubsystem->GetViewers();
+		EntitySubsystem.ForEachSharedFragment<FMassVisualizationLODSharedFragment>([&Viewers](FMassVisualizationLODSharedFragment& LODSharedFragment)
+		{
+			LODSharedFragment.LODCalculator.PrepareExecution(Viewers);
+		});
+	}
 
-void UMassVisualizationLODProcessor::PrepareExecution()
-{
-	check(World);
+	{
+		TRACE_CPUPROFILER_EVENT_SCOPE(CalculateLOD)
 
-	check(LODSubsystem);
-	const TArray<FViewerInfo>& Viewers = LODSubsystem->GetViewers();
-	LODCalculator.PrepareExecution(Viewers);
-}
+		auto CalculateLOD = [this](FMassExecutionContext& Context)
+		{
+			FMassVisualizationLODSharedFragment& LODSharedFragment = Context.GetMutableSharedFragment<FMassVisualizationLODSharedFragment>();
+			TArrayView<FMassRepresentationLODFragment> RepresentationLODList = Context.GetMutableFragmentView<FMassRepresentationLODFragment>();
+			TConstArrayView<FMassViewerInfoFragment> ViewerInfoList = Context.GetFragmentView<FMassViewerInfoFragment>();
+			LODSharedFragment.LODCalculator.CalculateLOD(Context, ViewerInfoList, RepresentationLODList);
+		};
+		CloseEntityQuery.ForEachEntityChunk(EntitySubsystem, Context, CalculateLOD);
+		FarEntityQuery.ForEachEntityChunk(EntitySubsystem, Context, CalculateLOD);
+	}
 
-void UMassVisualizationLODProcessor::Execute(UMassEntitySubsystem& EntitySubsystem, FMassExecutionContext& Context)
-{
-	ExecuteInternal(EntitySubsystem, Context);
+	{
+		TRACE_CPUPROFILER_EVENT_SCOPE(AdjustDistanceAndLODFromCount)
+		EntitySubsystem.ForEachSharedFragment<FMassVisualizationLODSharedFragment>([](FMassVisualizationLODSharedFragment& LODSharedFragment)
+		{
+			LODSharedFragment.bHasAdjustedDistancesFromCount = LODSharedFragment.LODCalculator.AdjustDistancesFromCount();
+		});
+
+		CloseEntityAdjustDistanceQuery.ForEachEntityChunk(EntitySubsystem, Context, [this](FMassExecutionContext& Context)
+		{
+			FMassVisualizationLODSharedFragment& LODSharedFragment = Context.GetMutableSharedFragment<FMassVisualizationLODSharedFragment>();
+			TConstArrayView<FMassViewerInfoFragment> ViewerInfoList = Context.GetFragmentView<FMassViewerInfoFragment>();
+			TArrayView<FMassRepresentationLODFragment> RepresentationLODList = Context.GetMutableFragmentView<FMassRepresentationLODFragment>();
+			LODSharedFragment.LODCalculator.AdjustLODFromCount(Context, ViewerInfoList, RepresentationLODList);
+		});
+		// Far entities do not need to maximize count
+	}
+
+	// Optional debug display
+	if (UE::MassRepresentation::bDebugRepresentationLOD)
+	{
+		TRACE_CPUPROFILER_EVENT_SCOPE(DebugDisplayLOD)
+		DebugEntityQuery.ForEachEntityChunk(EntitySubsystem, Context, [this](FMassExecutionContext& Context)
+		{
+			FMassVisualizationLODSharedFragment& LODSharedFragment = Context.GetMutableSharedFragment<FMassVisualizationLODSharedFragment>();
+			TConstArrayView<FMassRepresentationLODFragment> RepresentationLODList = Context.GetFragmentView<FMassRepresentationLODFragment>();
+			TConstArrayView<FTransformFragment> TransformList = Context.GetFragmentView<FTransformFragment>();
+			LODSharedFragment.LODCalculator.DebugDisplayLOD(Context, RepresentationLODList, TransformList, World);
+		});
+	}
 }
\ No newline at end of file
diff --git a/Engine/Plugins/Runtime/MassGameplay/Source/MassRepresentation/Private/MassVisualizationTrait.cpp b/Engine/Plugins/Runtime/MassGameplay/Source/MassRepresentation/Private/MassVisualizationTrait.cpp
--- a/Engine/Plugins/Runtime/MassGameplay/Source/MassRepresentation/Private/MassVisualizationTrait.cpp
+++ b/Engine/Plugins/Runtime/MassGameplay/Source/MassRepresentation/Private/MassVisualizationTrait.cpp
@@ -18,6 +18,25 @@ UMassVisualizationTrait::UMassVisualizationTrait()
 	Params.LODRepresentation[EMassLOD::Medium] = EMassRepresentationType::LowResSpawnedActor;
 	Params.LODRepresentation[EMassLOD::Low] = EMassRepresentationType::StaticMeshInstance;
 	Params.LODRepresentation[EMassLOD::Off] = EMassRepresentationType::None;
+
+	LODParams.BaseLODDistance[EMassLOD::High] = 0.f;
+	LODParams.BaseLODDistance[EMassLOD::Medium] = 1000.f;
+	LODParams.BaseLODDistance[EMassLOD::Low] = 2500.f;
+	LODParams.BaseLODDistance[EMassLOD::Off] = 10000.f;
+
+	LODParams.VisibleLODDistance[EMassLOD::High] = 0.f;
+	LODParams.VisibleLODDistance[EMassLOD::Medium] = 2000.f;
+	LODParams.VisibleLODDistance[EMassLOD::Low] = 4000.f;
+	LODParams.VisibleLODDistance[EMassLOD::Off] = 10000.f;
+
+	LODParams.LODMaxCount[EMassLOD::High] = 50;
+	LODParams.LODMaxCount[EMassLOD::Medium] = 100;
+	LODParams.LODMaxCount[EMassLOD::Low] = 500;
+	LODParams.LODMaxCount[EMassLOD::Off] = 0;
+
+	LODParams.BufferHysteresisOnDistancePercentage = 10.0f;
+	LODParams.DistanceToFrustum = 0.0f;
+	LODParams.DistanceToFrustumHysteresis = 0.0f;
 }
 
 void UMassVisualizationTrait::BuildTemplate(FMassEntityTemplateBuildContext& BuildContext, UWorld& World) const
@@ -65,6 +84,12 @@ void UMassVisualizationTrait::BuildTemplate(FMassEntityTemplateBuildContext& Bui
 	RepresentationFragment.HighResTemplateActorIndex = HighResTemplateActor.Get() ? RepresentationSubsystem->FindOrAddTemplateActor(HighResTemplateActor.Get()) : INDEX_NONE;
 	RepresentationFragment.LowResTemplateActorIndex = LowResTemplateActor.Get() ? RepresentationSubsystem->FindOrAddTemplateActor(LowResTemplateActor.Get()) : INDEX_NONE;
 
+	uint32 LODParamsHash = UE::StructUtils::GetStructCrc32(FConstStructView::Make(LODParams));
+	FConstSharedStruct LODParamsFragment = EntitySubsystem->GetOrCreateConstSharedFragment<FMassVisualizationLODParameters>(LODParamsHash, LODParams);
+	BuildContext.AddConstSharedFragment(LODParamsFragment);
+	FSharedStruct LODSharedFragment = EntitySubsystem->GetOrCreateSharedFragment<FMassVisualizationLODSharedFragment>(LODParamsHash, LODParams);
+	BuildContext.AddSharedFragment(LODSharedFragment);
+
 	BuildContext.AddFragment<FMassRepresentationLODFragment>();
 	BuildContext.AddTag<FMassVisibilityCulledByDistanceTag>();
 	BuildContext.AddChunkFragment<FMassVisualizationChunkFragment>();
diff --git a/Engine/Plugins/Runtime/MassGameplay/Source/MassRepresentation/Public/MassRepresentationFragments.h b/Engine/Plugins/Runtime/MassGameplay/Source/MassRepresentation/Public/MassRepresentationFragments.h
--- a/Engine/Plugins/Runtime/MassGameplay/Source/MassRepresentation/Public/MassRepresentationFragments.h
+++ b/Engine/Plugins/Runtime/MassGameplay/Source/MassRepresentation/Public/MassRepresentationFragments.h
@@ -8,6 +8,7 @@
 #include "MassActorSpawnerSubsystem.h"
 #include "MassRepresentationTypes.h"
 #include "MassRepresentationActorManagement.h"
+#include "MassLODCalculator.h"
 
 #include "MassRepresentationFragments.generated.h"
 
@@ -130,3 +131,40 @@ inline void FMassRepresentationParameters::ComputeCachedValues() const
 		CachedRepresentationActorManagement = UMassRepresentationActorManagement::StaticClass()->GetDefaultObject<UMassRepresentationActorManagement>();
 	}
 }
+
+USTRUCT()
+struct FMassVisualizationLODParameters : public FMassSharedFragment
+{
+	GENERATED_BODY()
+
+	/** Distances where each LOD becomes relevant */
+	UPROPERTY(EditAnywhere, Category = "Mass|LOD", config)
+	float BaseLODDistance[EMassLOD::Max];
+	UPROPERTY(EditAnywhere, Category = "Mass|LOD", config)
+	float VisibleLODDistance[EMassLOD::Max];
+	UPROPERTY(EditAnywhere, Category = "Mass|LOD", meta = (ClampMin = "0.0", UIMin = "0.0"), config)
+	float BufferHysteresisOnDistancePercentage = 10.0f;
+
+	/** Maximum limit for each entity per LOD */
+	UPROPERTY(EditAnywhere, Category = "Mass|LOD", config)
+	int32 LODMaxCount[EMassLOD::Max];
+
+	/** How far away from frustum does this entities are considered visible */
+	UPROPERTY(EditAnywhere, Category = "Mass|LOD", meta = (ClampMin = "0.0", UIMin = "0.0"), config)
+	float DistanceToFrustum = 0.0f;
+	/** Once visible how much further than DistanceToFrustum does the entities need to be before being cull again */
+	UPROPERTY(EditAnywhere, Category = "Mass|LOD", meta = (ClampMin = "0.0", UIMin = "0.0"), config)
+	float DistanceToFrustumHysteresis = 0.0f;
+};
+
+USTRUCT()
+struct FMassVisualizationLODSharedFragment : public FMassSharedFragment
+{
+	GENERATED_BODY()
+
+	FMassVisualizationLODSharedFragment() = default;
+	FMassVisualizationLODSharedFragment(const FMassVisualizationLODParameters& LODParams);
+
+	TMassLODCalculator<FMassRepresentationLODLogic> LODCalculator;
+	bool bHasAdjustedDistancesFromCount = false;
+};
diff --git a/Engine/Plugins/Runtime/MassGameplay/Source/MassRepresentation/Public/MassVisualizationLODProcessor.h b/Engine/Plugins/Runtime/MassGameplay/Source/MassRepresentation/Public/MassVisualizationLODProcessor.h
--- a/Engine/Plugins/Runtime/MassGameplay/Source/MassRepresentation/Public/MassVisualizationLODProcessor.h
+++ b/Engine/Plugins/Runtime/MassGameplay/Source/MassRepresentation/Public/MassVisualizationLODProcessor.h
@@ -21,24 +21,12 @@ protected:
 	/** Configure the owned FMassEntityQuery instances to express processor's requirements */
 	virtual void ConfigureQueries() override;
 
-	/**
-	 * Initialize the processor 
-	 * @param Owner of the Processor
-	 */
-	virtual void Initialize(UObject& Owner) override;
-
 	/** 
 	 * Execution method for this processor 
 	 * @param EntitySubsystem is the system to execute the lambdas on each entity chunk
 	 * @param Context is the execution context to be passed when executing the lambdas
 	 */
 	virtual void Execute(UMassEntitySubsystem& EntitySubsystem, FMassExecutionContext& Context) override;
-
-	void PrepareExecution();
-
-	template <typename TMassViewerLODInfoFragment = FMassViewerInfoFragment>
-	void ExecuteInternal(UMassEntitySubsystem& EntitySubsystem, FMassExecutionContext& Context);
-
 	
 	/**
 	 * Forces Off LOD on all calculation
@@ -47,107 +35,10 @@ protected:
 	void ForceOffLOD(bool bForce) { bForceOFFLOD = bForce; }
 
 protected:
-	/** Distances where each LOD becomes relevant */
-	UPROPERTY(EditAnywhere, Category = "Mass|LOD", config)
-	float BaseLODDistance[EMassLOD::Max];
-	UPROPERTY(EditAnywhere, Category = "Mass|LOD", config)
-	float VisibleLODDistance[EMassLOD::Max];
-	UPROPERTY(EditAnywhere, Category = "Mass|LOD", meta = (ClampMin = "0.0", UIMin = "0.0"), config)
-	float BufferHysteresisOnDistancePercentage = 10.0f;
-
-	/** Maximum limit for each entity per LOD */
-	UPROPERTY(EditAnywhere, Category = "Mass|LOD", config)
-	int32 LODMaxCount[EMassLOD::Max];
-
-	/** How far away from frustum does this entities are considered visible */
-	UPROPERTY(EditAnywhere, Category = "Mass|LOD", meta = (ClampMin = "0.0", UIMin = "0.0"), config)
-	float DistanceToFrustum = 0.0f;
-	/** Once visible how much further than DistanceToFrustum does the entities need to be before being cull again */
-	UPROPERTY(EditAnywhere, Category = "Mass|LOD", meta = (ClampMin = "0.0", UIMin = "0.0"), config)
-	float DistanceToFrustumHysteresis = 0.0f;
-
-	TMassLODCalculator<FMassRepresentationLODLogic> LODCalculator;
-
 	FMassEntityQuery CloseEntityQuery;
+	FMassEntityQuery CloseEntityAdjustDistanceQuery;
 	FMassEntityQuery FarEntityQuery;
+	FMassEntityQuery DebugEntityQuery;
 
 	bool bForceOFFLOD = false;
-};
-
-namespace UE::MassRepresentation
-{
-	extern MASSREPRESENTATION_API int32 bDebugRepresentationLOD;
-} // UE::MassRepresentation
-
-template <typename TMassViewerLODInfoFragment>
-void UMassVisualizationLODProcessor::ExecuteInternal(UMassEntitySubsystem& EntitySubsystem, FMassExecutionContext& Context)
-{
-	if (bForceOFFLOD)
-	{
-		CloseEntityQuery.ForEachEntityChunk(EntitySubsystem, Context, [this](FMassExecutionContext& Context)
-		{
-			TArrayView<FMassRepresentationLODFragment> RepresentationLODList = Context.GetMutableFragmentView<FMassRepresentationLODFragment>();
-			LODCalculator.ForceOffLOD(Context, RepresentationLODList);
-		});
-		return;
-	}
-
-	{
-		TRACE_CPUPROFILER_EVENT_SCOPE(PrepareExecution)
-
-		PrepareExecution();
-	}
-
-	{
-		TRACE_CPUPROFILER_EVENT_SCOPE(CalculateLOD)
-
-		// @todo this is the only block in this processor actually using FMassRepresentationFragment. Consider refactoring it into a separate Processor
-		// the "debug" section below is non-production, should be separated as well
-		auto CalculateLOD = [this](FMassExecutionContext& Context)
-		{
-			TArrayView<FMassRepresentationLODFragment> RepresentationLODList = Context.GetMutableFragmentView<FMassRepresentationLODFragment>();
-			TConstArrayView<TMassViewerLODInfoFragment> ViewerInfoList = Context.GetFragmentView<TMassViewerLODInfoFragment>();
-			LODCalculator.CalculateLOD(Context, ViewerInfoList, RepresentationLODList);
-		};
-		CloseEntityQuery.ForEachEntityChunk(EntitySubsystem, Context, CalculateLOD);
-
-		// @todo this is the only block in this processor actually using FMassRepresentationFragment. Consider refactoring it into a separate Processor
-		// the "debug" section below is non-production, should be separated as well
-		FarEntityQuery.SetChunkFilter(&FMassVisualizationChunkFragment::ShouldUpdateVisualizationForChunk);
-		FarEntityQuery.ForEachEntityChunk(EntitySubsystem, Context, [this, &CalculateLOD](FMassExecutionContext& Context)
-			{
-				CalculateLOD(Context);
-			});
-		FarEntityQuery.ClearChunkFilter();
-	}
-
-	{
-		TRACE_CPUPROFILER_EVENT_SCOPE(AdjustDistanceAndLODFromCount)
-		if (LODCalculator.AdjustDistancesFromCount())
-		{
-			CloseEntityQuery.ForEachEntityChunk(EntitySubsystem, Context, [this](FMassExecutionContext& Context)
-			{
-				TConstArrayView<TMassViewerLODInfoFragment> ViewerInfoList = Context.GetFragmentView<TMassViewerLODInfoFragment>();
-				TArrayView<FMassRepresentationLODFragment> RepresentationLODList = Context.GetMutableFragmentView<FMassRepresentationLODFragment>();
-				LODCalculator.AdjustLODFromCount(Context, ViewerInfoList, RepresentationLODList);
-			});
-			// Far entities do not need to maximice count
-		}
-	}
-
-	// Optional debug display
-	if(UE::MassRepresentation::bDebugRepresentationLOD)
-	{
-		TRACE_CPUPROFILER_EVENT_SCOPE(DebugDisplayLOD)
-
-		auto DebugDisplayLOD = [this](FMassExecutionContext& Context)
-		{
-			TConstArrayView<FMassRepresentationLODFragment> RepresentationLODList = Context.GetFragmentView<FMassRepresentationLODFragment>();
-			TConstArrayView<FTransformFragment> TransformList = Context.GetFragmentView<FTransformFragment>();
-			LODCalculator.DebugDisplayLOD(Context, RepresentationLODList, TransformList, World);
-		};
-
-		CloseEntityQuery.ForEachEntityChunk(EntitySubsystem, Context, DebugDisplayLOD);
-		FarEntityQuery.ForEachEntityChunk(EntitySubsystem, Context, DebugDisplayLOD);
-	}
-}
\ No newline at end of file
+};
\ No newline at end of file
diff --git a/Engine/Plugins/Runtime/MassGameplay/Source/MassRepresentation/Public/MassVisualizationTrait.h b/Engine/Plugins/Runtime/MassGameplay/Source/MassRepresentation/Public/MassVisualizationTrait.h
--- a/Engine/Plugins/Runtime/MassGameplay/Source/MassRepresentation/Public/MassVisualizationTrait.h
+++ b/Engine/Plugins/Runtime/MassGameplay/Source/MassRepresentation/Public/MassVisualizationTrait.h
@@ -38,7 +38,11 @@ public:
 	UPROPERTY(EditAnywhere, Category = "Mass|Visual")
 	TSubclassOf<UMassRepresentationSubsystem> RepresentationSubsystemClass;
 
-	/** Configuration parameters for the representaion processor */
+	/** Configuration parameters for the representation processor */
 	UPROPERTY(EditAnywhere, Category = "Mass|Visual")
 	FMassRepresentationParameters Params;
+
+	/** Configuration parameters for the visualization LOD processor */
+	UPROPERTY(EditAnywhere, Category = "Mass|Visual")
+	FMassVisualizationLODParameters LODParams;
 };
