Fix some perf regression in sample
* Deleted un needed SampleCrowdVisualizationProcessor, saves 0.35 over 0.7ms) the visualization was done twice.
* Re-parallelize the LOD collection of the trafic with the crowd (Saves ~1.6ms on my local PC over 7.5ms)
* Augmented the size of the chunks to 128k because since we enable LWC it have afftected the number of entities per chunk (Mostly on the traffic saves ~1ms over 6ms)
* Found out that we were collecting all the parked vehicles LOD info because they were not using SimulationLOD, fixed by changing the method logics we call to know if we will need those distances.(Saves ~0.4ms over 0.7ms)
* Cleaned up some of the deleted parameters in mass ini file

#rnx

diff --git a/Engine/Plugins/AI/MassCrowd/Source/MassCrowd/Private/MassCrowdVisualizationLODProcessor.cpp b/Engine/Plugins/AI/MassCrowd/Source/MassCrowd/Private/MassCrowdVisualizationLODProcessor.cpp
--- a/Engine/Plugins/AI/MassCrowd/Source/MassCrowd/Private/MassCrowdVisualizationLODProcessor.cpp
+++ b/Engine/Plugins/AI/MassCrowd/Source/MassCrowd/Private/MassCrowdVisualizationLODProcessor.cpp
@@ -19,24 +19,11 @@ namespace UE::MassCrowd
 		FAutoConsoleVariableRef(TEXT("ai.debug.ShowISMUnderSpecifiedRange"), bDebugShowISMUnderSpecifiedRange, TEXT("Show ISM under a specified range (meters)"), ECVF_Cheat)
 	};
 
-// #if WITH_EDITOR
-// 	int32 DebugCrowdMaxCountHigh = -1;
-// 	int32 DebugCrowdMaxCountMedium = -1;
-// 	int32 DebugCrowdMaxCountLow = -1;
-// 
-// 	FAutoConsoleVariableRef EditorConsoleVariables[] =
-// 	{
-// 		FAutoConsoleVariableRef(TEXT("ai.debug.CrowdMaxCountHigh"), DebugCrowdMaxCountHigh, TEXT("Forced MASS crowd max count high"), ECVF_Cheat),
-// 		FAutoConsoleVariableRef(TEXT("ai.debug.CrowdMaxCountMedium"), DebugCrowdMaxCountMedium, TEXT("Forced MASS crowd max count medium"), ECVF_Cheat),
-// 		FAutoConsoleVariableRef(TEXT("ai.debug.CrowdMaxCountLow"), DebugCrowdMaxCountLow, TEXT("Forced MASS crowd max count low"), ECVF_Cheat)
-// 	};
-// 
-// 	int32 LastDebugCrowdMaxCountHigh = -1;
-// 	int32 LastDebugCrowdMaxCountMedium = -1;
-// 	int32 LastDebugCrowdMaxCountLow = -1;
-// #endif // WITH_EDITOR
 } // UE::MassCrowd
 
+//----------------------------------------------------------------------//
+// UMassCrowdVisualizationLODProcessor
+//----------------------------------------------------------------------//
 UMassCrowdVisualizationLODProcessor::UMassCrowdVisualizationLODProcessor()
 {
 	bAutoRegisterWithProcessingPhases = true;
@@ -45,12 +32,6 @@ UMassCrowdVisualizationLODProcessor::UMassCrowdVisualizationLODProcessor()
 
 	ExecutionOrder.ExecuteInGroup = UE::Mass::ProcessorGroupNames::LOD;
 	ExecutionOrder.ExecuteAfter.Add(UE::Mass::ProcessorGroupNames::LODCollector);
-
-// #if WITH_EDITOR
-// 	UE::MassCrowd::LastDebugCrowdMaxCountHigh = -1;
-// 	UE::MassCrowd::LastDebugCrowdMaxCountMedium = -1;
-// 	UE::MassCrowd::LastDebugCrowdMaxCountLow = -1;
-// #endif // WITH_EDITOR
 }
 
 void UMassCrowdVisualizationLODProcessor::ConfigureQueries()
@@ -61,6 +42,8 @@ void UMassCrowdVisualizationLODProcessor::ConfigureQueries()
 	CloseEntityAdjustDistanceQuery.AddTagRequirement<FMassCrowdTag>(EMassFragmentPresence::All);
 	FarEntityQuery.AddTagRequirement<FMassCrowdTag>(EMassFragmentPresence::All);
 	DebugEntityQuery.AddTagRequirement<FMassCrowdTag>(EMassFragmentPresence::All);
+
+	FilterTag = FMassCrowdTag::StaticStruct();
 }
 
 void UMassCrowdVisualizationLODProcessor::Execute(UMassEntitySubsystem& EntitySubsystem, FMassExecutionContext& Context)
@@ -107,47 +90,22 @@ void UMassCrowdVisualizationLODProcessor::Execute(UMassEntitySubsystem& EntitySu
 			}
 		});
 	}
+}
+
+//----------------------------------------------------------------------//
+// UMassCrowdLODCollectorProcessor
+//----------------------------------------------------------------------//
+UMassCrowdLODCollectorProcessor::UMassCrowdLODCollectorProcessor()
+{
+	bAutoRegisterWithProcessingPhases = true;
+}
+
+void UMassCrowdLODCollectorProcessor::ConfigureQueries()
+{
+	Super::ConfigureQueries();
 
-// @todo find a way to do this now with shared fragments
-// #if WITH_EDITOR
-// 	if (UE::MassCrowd::LastDebugCrowdMaxCountHigh != UE::MassCrowd::DebugCrowdMaxCountHigh)
-// 	{
-// 		if (UE::MassCrowd::DebugCrowdMaxCountHigh >= 0)
-// 		{
-// 			LODMaxCount[0] = UE::MassCrowd::DebugCrowdMaxCountHigh;
-// 		}
-// 		else
-// 		{
-// 			LODMaxCount[0] = GetClass()->GetDefaultObject<UMassCrowdVisualizationLODProcessor>()->LODMaxCount[0];
-// 		}
-// 		UE::MassCrowd::LastDebugCrowdMaxCountHigh = UE::MassCrowd::DebugCrowdMaxCountHigh;
-// 		LODCalculator.Initialize(BaseLODDistance, BufferHysteresisOnDistancePercentage / 100.f, LODMaxCount, nullptr, DistanceToFrustum, DistanceToFrustumHysteresis, VisibleLODDistance);
-// 	}
-// 	if (UE::MassCrowd::LastDebugCrowdMaxCountMedium != UE::MassCrowd::DebugCrowdMaxCountMedium)
-// 	{
-// 		if (UE::MassCrowd::DebugCrowdMaxCountMedium >= 0)
-// 		{
-// 			LODMaxCount[1] = UE::MassCrowd::DebugCrowdMaxCountMedium;
-// 		}
-// 		else
-// 		{
-// 			LODMaxCount[1] = GetClass()->GetDefaultObject<UMassCrowdVisualizationLODProcessor>()->LODMaxCount[1];
-// 		}
-// 		UE::MassCrowd::LastDebugCrowdMaxCountMedium = UE::MassCrowd::DebugCrowdMaxCountMedium;
-// 		LODCalculator.Initialize(BaseLODDistance, BufferHysteresisOnDistancePercentage / 100.f, LODMaxCount, nullptr, DistanceToFrustum, DistanceToFrustumHysteresis, VisibleLODDistance);
-// 	}
-// 	if (UE::MassCrowd::LastDebugCrowdMaxCountLow != UE::MassCrowd::DebugCrowdMaxCountLow)
-// 	{
-// 		if (UE::MassCrowd::DebugCrowdMaxCountLow >= 0)
-// 		{
-// 			LODMaxCount[2] = UE::MassCrowd::DebugCrowdMaxCountLow;
-// 		}
-// 		else
-// 		{
-// 			LODMaxCount[2] = GetClass()->GetDefaultObject<UMassCrowdVisualizationLODProcessor>()->LODMaxCount[2];
-// 		}
-// 		UE::MassCrowd::LastDebugCrowdMaxCountLow = UE::MassCrowd::DebugCrowdMaxCountLow;
-// 		LODCalculator.Initialize(BaseLODDistance, BufferHysteresisOnDistancePercentage / 100.f, LODMaxCount, nullptr, DistanceToFrustum, DistanceToFrustumHysteresis, VisibleLODDistance);
-// 	}
-// #endif // WITH_EDITOR
+	EntityQuery_VisibleRangeAndOnLOD.AddTagRequirement<FMassCrowdTag>(EMassFragmentPresence::All);
+	EntityQuery_VisibleRangeOnly.AddTagRequirement<FMassCrowdTag>(EMassFragmentPresence::All);
+	EntityQuery_OnLODOnly.AddTagRequirement<FMassCrowdTag>(EMassFragmentPresence::All);
+	EntityQuery_NotVisibleRangeAndOffLOD.AddTagRequirement<FMassCrowdTag>(EMassFragmentPresence::All);
 }
diff --git a/Engine/Plugins/AI/MassCrowd/Source/MassCrowd/Private/MassCrowdVisualizationTrait.cpp b/Engine/Plugins/AI/MassCrowd/Source/MassCrowd/Private/MassCrowdVisualizationTrait.cpp
--- a/Engine/Plugins/AI/MassCrowd/Source/MassCrowd/Private/MassCrowdVisualizationTrait.cpp
+++ b/Engine/Plugins/AI/MassCrowd/Source/MassCrowd/Private/MassCrowdVisualizationTrait.cpp
@@ -4,6 +4,7 @@
 #include "MassCrowdRepresentationSubsystem.h"
 #include "MassCrowdVisualizationProcessor.h"
 #include "MassCrowdRepresentationActorManagement.h"
+#include "MassCrowdFragments.h"
 
 UMassCrowdVisualizationTrait::UMassCrowdVisualizationTrait()
 {
@@ -39,4 +40,6 @@ UMassCrowdVisualizationTrait::UMassCrowdVisualizationTrait()
 	LODParams.BufferHysteresisOnDistancePercentage = 20.0f;
 	LODParams.DistanceToFrustum = 0.0f;
 	LODParams.DistanceToFrustumHysteresis = 0.0f;
+
+	LODParams.FilterTag = FMassCrowdTag::StaticStruct();
 }
diff --git a/Engine/Plugins/AI/MassCrowd/Source/MassCrowd/Public/MassCrowdVisualizationLODProcessor.h b/Engine/Plugins/AI/MassCrowd/Source/MassCrowd/Public/MassCrowdVisualizationLODProcessor.h
--- a/Engine/Plugins/AI/MassCrowd/Source/MassCrowd/Public/MassCrowdVisualizationLODProcessor.h
+++ b/Engine/Plugins/AI/MassCrowd/Source/MassCrowd/Public/MassCrowdVisualizationLODProcessor.h
@@ -3,11 +3,11 @@
 #pragma once
 
 #include "MassVisualizationLODProcessor.h"
+#include "MassLODCollectorProcessor.h"
 #include "MassCrowdVisualizationLODProcessor.generated.h"
 
 /*
- * Created a crowd version for visualization of the UMassVisualizationLODProcessor as we want to use a custom version of the LOD info
- * that combines the visualization and simulation LOD.
+ * Created a crowd version for parallelization of the crowd with the traffic
  */
 UCLASS(meta=(DisplayName="Crowd visualization LOD"))
 class MASSCROWD_API UMassCrowdVisualizationLODProcessor : public UMassVisualizationLODProcessor
@@ -20,3 +20,18 @@ protected:
 	virtual void ConfigureQueries() override;
 	virtual void Execute(UMassEntitySubsystem& EntitySubsystem, FMassExecutionContext& Context) override;
 };
+
+/*
+ * Created a crowd version for parallelization of the crowd with the traffic
+ */
+UCLASS(meta = (DisplayName = "Crowd LOD Collection "))
+class MASSCROWD_API UMassCrowdLODCollectorProcessor : public UMassLODCollectorProcessor
+{
+	GENERATED_BODY()
+
+	UMassCrowdLODCollectorProcessor();
+
+protected:
+	virtual void ConfigureQueries() override;
+};
+
diff --git a/Engine/Plugins/Runtime/MassEntity/Source/MassEntity/Private/MassArchetypeData.h b/Engine/Plugins/Runtime/MassEntity/Source/MassEntity/Private/MassArchetypeData.h
--- a/Engine/Plugins/Runtime/MassEntity/Source/MassEntity/Private/MassArchetypeData.h
+++ b/Engine/Plugins/Runtime/MassEntity/Source/MassEntity/Private/MassArchetypeData.h
@@ -10,6 +10,11 @@ struct FMassExecutionContext;
 class FOutputDevice;
 struct FMassArchetypeEntityCollection;
 
+namespace UE::MassEntity
+{
+	constexpr int32 ChunkSize = 128*1024;
+}
+
 // This is one chunk within an archetype
 struct FMassArchetypeChunk
 {
@@ -205,7 +210,7 @@ public:
 
 	int32 GetNumEntities() const { return EntityMap.Num(); }
 
-	int32 GetChunkAllocSize() const { return 64*1024; }
+	int32 GetChunkAllocSize() const { return UE::MassEntity::ChunkSize; }
 
 	int32 GetChunkCount() const { return Chunks.Num(); }
 
diff --git a/Engine/Plugins/Runtime/MassGameplay/Source/MassLOD/Private/MassLODCollectorProcessor.cpp b/Engine/Plugins/Runtime/MassGameplay/Source/MassLOD/Private/MassLODCollectorProcessor.cpp
--- a/Engine/Plugins/Runtime/MassGameplay/Source/MassLOD/Private/MassLODCollectorProcessor.cpp
+++ b/Engine/Plugins/Runtime/MassGameplay/Source/MassLOD/Private/MassLODCollectorProcessor.cpp
@@ -8,6 +8,8 @@
 
 UMassLODCollectorProcessor::UMassLODCollectorProcessor()
 {
+	bAutoRegisterWithProcessingPhases = false;
+
 	ExecutionFlags = (int32)EProcessorExecutionFlags::All;
 
 	ExecutionOrder.ExecuteInGroup = UE::Mass::ProcessorGroupNames::LODCollector;
@@ -24,8 +26,8 @@ void UMassLODCollectorProcessor::ConfigureQueries()
 	BaseQuery.AddChunkRequirement<FMassVisualizationChunkFragment>(EMassFragmentAccess::ReadOnly, EMassFragmentPresence::Optional);
 	BaseQuery.SetChunkFilter([](const FMassExecutionContext& Context)
 	{
-		return FMassVisualizationChunkFragment::ShouldUpdateVisualizationForChunk(Context)
-			|| FMassSimulationVariableTickChunkFragment::ShouldTickChunkThisFrame(Context);
+		return FMassVisualizationChunkFragment::IsChunkHandledThisFrame(Context)
+			|| FMassSimulationVariableTickChunkFragment::IsChunkHandledThisFrame(Context);
 	});
 
 	EntityQuery_VisibleRangeAndOnLOD = BaseQuery;
diff --git a/Engine/Plugins/Runtime/MassGameplay/Source/MassLOD/Public/MassLODFragments.h b/Engine/Plugins/Runtime/MassGameplay/Source/MassLOD/Public/MassLODFragments.h
--- a/Engine/Plugins/Runtime/MassGameplay/Source/MassLOD/Public/MassLODFragments.h
+++ b/Engine/Plugins/Runtime/MassGameplay/Source/MassLOD/Public/MassLODFragments.h
@@ -141,11 +141,29 @@ struct MASSLOD_API FMassVisualizationChunkFragment : public FMassChunkFragment
 		return (Visibility != EMassVisibility::CulledByDistance && Visibility != EMassVisibility::CulledByFrustum) || bContainsNewlyVisibleEntity;
 	}
 
+	/**
+	 * IsChunkHandledThisFrame
+	 *
+	 * This function is used by LOD collector query chunk filters to check that visual LOD will be updated this frame. 
+	 * It defaults to false (no LOD update), if visualization chunk fragment is NOT present.
+	 * 
+	 * @param Context of the execution from the entity sub system
+	 * @return true if the visual LOD will be updated this frame
+	 */
+	static bool IsChunkHandledThisFrame(const FMassExecutionContext& Context)
+	{
+		const FMassVisualizationChunkFragment* ChunkFragment = Context.GetChunkFragmentPtr<FMassVisualizationChunkFragment>();
+		return ChunkFragment != nullptr && ChunkFragment->ShouldUpdateVisualization();
+	}
+
 	/**
 	 * ShouldUpdateVisualizationForChunk
+	 * 
+	 * This function is used by query chunk filters in processors that require variable visual LOD update. 
+	 * It defaults to true (always updating) if visualization chunk fragment is NOT present.
 	 *
-	* @param Context of the execution from the entity sub system
-	 * @return bool true if we should update the representation type for this chunk
+	 * @param Context of the execution from the entity sub system
+	 * @return true if the chunk should update the visual this frame
 	 */
 	static bool ShouldUpdateVisualizationForChunk(const FMassExecutionContext& Context)
 	{
diff --git a/Engine/Plugins/Runtime/MassGameplay/Source/MassLOD/Public/MassSimulationLOD.h b/Engine/Plugins/Runtime/MassGameplay/Source/MassLOD/Public/MassSimulationLOD.h
--- a/Engine/Plugins/Runtime/MassGameplay/Source/MassLOD/Public/MassSimulationLOD.h
+++ b/Engine/Plugins/Runtime/MassGameplay/Source/MassLOD/Public/MassSimulationLOD.h
@@ -43,8 +43,26 @@ struct MASSLOD_API FMassSimulationVariableTickChunkFragment : public FMassVariab
 	GENERATED_BODY();
 
 	/**
-	 * @return if Context contains an instance of FMassSimulationVariableTickChunkFragment then the function returns
-	 *	the value of FMassVariableTickChunkFragment.ShouldTickThisFrame; otherwise it returns DefaultValue.
+	 * IsChunkHandledThisFrame
+	 * 
+	 * This function is used by LOD collector query chunk filters to check if the Simulation LOD will be updated this frame.
+	 * It defaults to false (no LOD update), if simulation variable tick chunk fragment is NOT present.
+	 * 
+	 * @return true if the simulation LOD will be updated this frame
+	 */
+	static bool IsChunkHandledThisFrame(const FMassExecutionContext& Context)
+	{
+		const FMassSimulationVariableTickChunkFragment* ChunkFragment = Context.GetChunkFragmentPtr<FMassSimulationVariableTickChunkFragment>();
+		return ChunkFragment != nullptr && ChunkFragment->ShouldTickThisFrame();
+	}
+
+	/**
+	 * ShouldTickChunkThisFrame
+	 * 
+	 * This function is used by query chunk filters in processors that require variable rate ticking based on LOD.
+	 * It defaults to true (always ticking) if simulation variable tick chunk fragment is NOT present.
+	 * 
+	 * @return if the chunk should be ticked this frame
 	 */
 	static bool ShouldTickChunkThisFrame(const FMassExecutionContext& Context)
 	{
