Fixed cooked editor QAGame compilation

diff --git a/Engine/Plugins/AI/MassCrowd/Source/MassCrowd/Private/MassCrowdBubble.cpp b/Engine/Plugins/AI/MassCrowd/Source/MassCrowd/Private/MassCrowdBubble.cpp
--- a/Engine/Plugins/AI/MassCrowd/Source/MassCrowd/Private/MassCrowdBubble.cpp
+++ b/Engine/Plugins/AI/MassCrowd/Source/MassCrowd/Private/MassCrowdBubble.cpp
@@ -68,6 +68,7 @@ void FMassCrowdClientBubbleHandler::DebugValidateBubbleOnServer()
 {
 	Super::DebugValidateBubbleOnServer();
 
+#if UE_REPLICATION_COMPILE_SERVER_CODE
 	if (UE::Mass::Crowd::bDebugReplicationPositions)
 	{
 		const FMassEntityManager& EntityManager = Serializer->GetEntityManagerChecked();
@@ -86,6 +87,7 @@ void FMassCrowdClientBubbleHandler::DebugValidateBubbleOnServer()
 			UE::Mass::Crowd::DebugDrawReplicatedAgent(LookupData.Entity, EntityManager);
 		}
 	}
+#endif // UE_REPLICATION_COMPILE_SERVER_CODE
 }
 #endif // WITH_MASSGAMEPLAY_DEBUG && WITH_EDITOR
 
