Added Mass traits tag requirement to help people figure out why a some trait aren't working like the CrowdVisualization.

diff --git a/Engine/Plugins/AI/MassCrowd/Source/MassCrowd/Private/MassCrowdVisualizationTrait.cpp b/Engine/Plugins/AI/MassCrowd/Source/MassCrowd/Private/MassCrowdVisualizationTrait.cpp
--- a/Engine/Plugins/AI/MassCrowd/Source/MassCrowd/Private/MassCrowdVisualizationTrait.cpp
+++ b/Engine/Plugins/AI/MassCrowd/Source/MassCrowd/Private/MassCrowdVisualizationTrait.cpp
@@ -5,6 +5,7 @@
 #include "MassCrowdVisualizationProcessor.h"
 #include "MassCrowdRepresentationActorManagement.h"
 #include "MassCrowdFragments.h"
+#include "MassEntityTemplateRegistry.h"
 
 UMassCrowdVisualizationTrait::UMassCrowdVisualizationTrait()
 {
@@ -23,14 +24,14 @@ UMassCrowdVisualizationTrait::UMassCrowdVisualizationTrait()
 	Params.NotVisibleUpdateRate = 0.5f;
 
 	LODParams.BaseLODDistance[EMassLOD::High] = 0.f;
-	LODParams.BaseLODDistance[EMassLOD::Medium] = 0.f;
-	LODParams.BaseLODDistance[EMassLOD::Low] = 3000.f;
-	LODParams.BaseLODDistance[EMassLOD::Off] = 6000.f;
+	LODParams.BaseLODDistance[EMassLOD::Medium] = 500.f;
+	LODParams.BaseLODDistance[EMassLOD::Low] = 1000.f;
+	LODParams.BaseLODDistance[EMassLOD::Off] = 5000.f;
 
 	LODParams.VisibleLODDistance[EMassLOD::High] = 0.f;
-	LODParams.VisibleLODDistance[EMassLOD::Medium] = 0.f;
-	LODParams.VisibleLODDistance[EMassLOD::Low] = 6000.f;
-	LODParams.VisibleLODDistance[EMassLOD::Off] = 50000.f;
+	LODParams.VisibleLODDistance[EMassLOD::Medium] = 1000.f;
+	LODParams.VisibleLODDistance[EMassLOD::Low] = 5000.f;
+	LODParams.VisibleLODDistance[EMassLOD::Off] = 10000.f;
 
 	LODParams.LODMaxCount[EMassLOD::High] = 10;
 	LODParams.LODMaxCount[EMassLOD::Medium] = 20;
@@ -43,3 +44,11 @@ UMassCrowdVisualizationTrait::UMassCrowdVisualizationTrait()
 
 	LODParams.FilterTag = FMassCrowdTag::StaticStruct();
 }
+
+void UMassCrowdVisualizationTrait::BuildTemplate(FMassEntityTemplateBuildContext& BuildContext,
+	const UWorld& World) const
+{
+	Super::BuildTemplate(BuildContext, World);
+
+	BuildContext.RequireTag<FMassCrowdTag>();
+}
diff --git a/Engine/Plugins/AI/MassCrowd/Source/MassCrowd/Public/MassCrowdVisualizationTrait.h b/Engine/Plugins/AI/MassCrowd/Source/MassCrowd/Public/MassCrowdVisualizationTrait.h
--- a/Engine/Plugins/AI/MassCrowd/Source/MassCrowd/Public/MassCrowdVisualizationTrait.h
+++ b/Engine/Plugins/AI/MassCrowd/Source/MassCrowd/Public/MassCrowdVisualizationTrait.h
@@ -12,4 +12,5 @@ class MASSCROWD_API UMassCrowdVisualizationTrait : public UMassVisualizationTrai
 public:
 	UMassCrowdVisualizationTrait();
 
+	virtual void BuildTemplate(FMassEntityTemplateBuildContext& BuildContext, const UWorld& World) const override;
 };
diff --git a/Engine/Plugins/Runtime/MassGameplay/Source/MassSpawner/Public/MassEntityTemplate.h b/Engine/Plugins/Runtime/MassGameplay/Source/MassSpawner/Public/MassEntityTemplate.h
--- a/Engine/Plugins/Runtime/MassGameplay/Source/MassSpawner/Public/MassEntityTemplate.h
+++ b/Engine/Plugins/Runtime/MassGameplay/Source/MassSpawner/Public/MassEntityTemplate.h
@@ -140,11 +140,13 @@ struct MASSSPAWNER_API FMassEntityTemplate
 	template<typename T>
 	void AddTag()
 	{
+		static_assert(TIsDerivedFrom<T, FMassTag>::IsDerived, "Given struct doesn't represent a valid mass tag type. Make sure to inherit from FMassTag or one of its child-types.");
 		Composition.Tags.Add<T>();
 	}
 	
 	void AddTag(const UScriptStruct& TagType)
 	{
+		checkf(TagType.IsChildOf(FMassFragment::StaticStruct()), TEXT("Given struct doesn't represent a valid mass tag type. Make sure to inherit from FMassTag or one of its child-types."));
 		Composition.Tags.Add(TagType);
 	}
 
diff --git a/Engine/Plugins/Runtime/MassGameplay/Source/MassSpawner/Public/MassEntityTemplateRegistry.h b/Engine/Plugins/Runtime/MassGameplay/Source/MassSpawner/Public/MassEntityTemplateRegistry.h
--- a/Engine/Plugins/Runtime/MassGameplay/Source/MassSpawner/Public/MassEntityTemplateRegistry.h
+++ b/Engine/Plugins/Runtime/MassGameplay/Source/MassSpawner/Public/MassEntityTemplateRegistry.h
@@ -54,12 +54,14 @@ struct FMassEntityTemplateBuildContext
 	{
 		// Tags can be added by multiple traits, so they do not follow the same rules as fragments
 		Template.AddTag<T>();
+		TypeAdded(*T::StaticStruct());
 	}
 
 	void AddTag(const UScriptStruct& TagType)
 	{
 		// Tags can be added by multiple traits, so they do not follow the same rules as fragments
 		Template.AddTag(TagType);
+		TypeAdded(TagType);
 	}
 
 	template<typename T>
@@ -141,6 +143,12 @@ struct FMassEntityTemplateBuildContext
 		AddDependency(T::StaticStruct());
 	}
 
+	template<typename T>
+	void RequireTag()
+	{
+		AddDependency(T::StaticStruct());
+	}
+
 	void AddDependency(const UStruct* Dependency)
 	{
 		TraitsDependencies.Add( {Dependency, BuildingTrait} );
