Added missing module dependency.

#rnx

diff --git a/Engine/Plugins/AI/MassCrowd/Source/MassCrowd/MassCrowd.Build.cs b/Engine/Plugins/AI/MassCrowd/Source/MassCrowd/MassCrowd.Build.cs
--- a/Engine/Plugins/AI/MassCrowd/Source/MassCrowd/MassCrowd.Build.cs
+++ b/Engine/Plugins/AI/MassCrowd/Source/MassCrowd/MassCrowd.Build.cs
@@ -36,6 +36,8 @@ namespace UnrealBuildTool.Rules
 					"ZoneGraphDebug"
 				}
 			);
+
+			SetupIrisSupport(Target);
 		}
 	}
 }
\ No newline at end of file
