A couple of trivial fixes and improvements to Mass code.

diff --git a/Engine/Plugins/AI/MassCrowd/Source/MassCrowd/Public/MassCrowdReplicator.h b/Engine/Plugins/AI/MassCrowd/Source/MassCrowd/Public/MassCrowdReplicator.h
--- a/Engine/Plugins/AI/MassCrowd/Source/MassCrowd/Public/MassCrowdReplicator.h
+++ b/Engine/Plugins/AI/MassCrowd/Source/MassCrowd/Public/MassCrowdReplicator.h
@@ -17,11 +17,11 @@ public:
 	 * Overridden to add specific entity query requirements for replication.
 	 * Usually we add replication processor handler requirements.
 	 */
-	void AddRequirements(FMassEntityQuery& EntityQuery) override;
+	virtual void AddRequirements(FMassEntityQuery& EntityQuery) override;
 	
 	/**
 	 * Overridden to process the client replication.
 	 * This methods should call CalculateClientReplication with the appropriate callback implementation.
 	 */
-	void ProcessClientReplication(FMassExecutionContext& Context, FMassReplicationContext& ReplicationContext) override;
+	virtual void ProcessClientReplication(FMassExecutionContext& Context, FMassReplicationContext& ReplicationContext) override;
 };
\ No newline at end of file
diff --git a/Engine/Plugins/Runtime/MassGameplay/Source/MassCommon/Public/MassCommonFragments.h b/Engine/Plugins/Runtime/MassGameplay/Source/MassCommon/Public/MassCommonFragments.h
--- a/Engine/Plugins/Runtime/MassGameplay/Source/MassCommon/Public/MassCommonFragments.h
+++ b/Engine/Plugins/Runtime/MassGameplay/Source/MassCommon/Public/MassCommonFragments.h
@@ -14,6 +14,11 @@ struct MASSCOMMON_API FTransformFragment : public FMassFragment
 {
 	GENERATED_BODY()
 
+	FTransformFragment() = default;
+	FTransformFragment(const FTransform& InTransform) 
+	: Transform(InTransform)
+	{}
+
 	const FTransform& GetTransform() const { return Transform; }
 	void SetTransform(const FTransform& InTransform) { Transform = InTransform; }
 	FTransform& GetMutableTransform() { return Transform; }
diff --git a/Engine/Plugins/Runtime/MassGameplay/Source/MassReplication/Private/MassReplicationProcessor.cpp b/Engine/Plugins/Runtime/MassGameplay/Source/MassReplication/Private/MassReplicationProcessor.cpp
--- a/Engine/Plugins/Runtime/MassGameplay/Source/MassReplication/Private/MassReplicationProcessor.cpp
+++ b/Engine/Plugins/Runtime/MassGameplay/Source/MassReplication/Private/MassReplicationProcessor.cpp
@@ -91,6 +91,12 @@ void UMassReplicationProcessor::PrepareExecution(FMassEntityManager& EntityManag
 
 	EntityManager.ForEachSharedFragment<FMassReplicationSharedFragment>([this](FMassReplicationSharedFragment& RepSharedFragment)
 	{
+		if (!ensureMsgf(RepSharedFragment.BubbleInfoClassHandle.IsValid()
+			, TEXT("BubbleInfoClassHandle is not valid which means no class has been indicated or the class used has not been registered pre creation of the handle.")))
+		{
+			return;
+		}
+
 		if (!RepSharedFragment.bEntityQueryInitialized)
 		{
 			RepSharedFragment.EntityQuery = EntityQuery;
diff --git a/Engine/Plugins/Runtime/MassGameplay/Source/MassReplication/Private/MassReplicationSubsystem.cpp b/Engine/Plugins/Runtime/MassGameplay/Source/MassReplication/Private/MassReplicationSubsystem.cpp
--- a/Engine/Plugins/Runtime/MassGameplay/Source/MassReplication/Private/MassReplicationSubsystem.cpp
+++ b/Engine/Plugins/Runtime/MassGameplay/Source/MassReplication/Private/MassReplicationSubsystem.cpp
@@ -10,6 +10,7 @@
 #include "MassClientBubbleInfoBase.h"
 #include "MassReplicationSettings.h"
 #include "MassEntityUtils.h"
+#include "VisualLogger/VisualLogger.h"
 
 
 uint32 UMassReplicationSubsystem::CurrentNetMassCounter = 0;
@@ -424,6 +425,9 @@ FMassBubbleInfoClassHandle UMassReplicationSubsystem::GetBubbleInfoClassHandle(c
 			break;
 		}
 	}
+
+	UE_CVLOG_UELOG(Handle.IsValid() == false, this, LogMassReplication, Error, TEXT("%hs failed to find required Bubble Info class"), __FUNCTION__);
+
 	return Handle;
 }
 
diff --git a/Engine/Plugins/Runtime/MassGameplay/Source/MassSpawner/Private/MassEntityTemplateRegistry.cpp b/Engine/Plugins/Runtime/MassGameplay/Source/MassSpawner/Private/MassEntityTemplateRegistry.cpp
--- a/Engine/Plugins/Runtime/MassGameplay/Source/MassSpawner/Private/MassEntityTemplateRegistry.cpp
+++ b/Engine/Plugins/Runtime/MassGameplay/Source/MassSpawner/Private/MassEntityTemplateRegistry.cpp
@@ -83,7 +83,7 @@ const TSharedRef<FMassEntityTemplate>& FMassEntityTemplateRegistry::FindOrAddTem
 	check(EntityManager);
 	const TSharedRef<FMassEntityTemplate>* ExistingTemplate = FindTemplateFromTemplateID(NewTemplateID);
 
-	if (!ensureMsgf(ExistingTemplate == nullptr, TEXT("")))
+	if (ExistingTemplate)
 	{
 		return *ExistingTemplate;
 	}
