Renamed some more `ai.debug` console commands to `mass.debug` (where appropriate).

diff --git a/Engine/Plugins/AI/MassCrowd/Source/MassCrowd/Private/MassCrowdBubble.cpp b/Engine/Plugins/AI/MassCrowd/Source/MassCrowd/Private/MassCrowdBubble.cpp
--- a/Engine/Plugins/AI/MassCrowd/Source/MassCrowd/Private/MassCrowdBubble.cpp
+++ b/Engine/Plugins/AI/MassCrowd/Source/MassCrowd/Private/MassCrowdBubble.cpp
@@ -8,7 +8,7 @@
 namespace UE::Mass::Crowd
 {
 	bool bDebugReplicationPositions = false;
-	FAutoConsoleVariableRef CVarbDebugReplication(TEXT("ai.debug.CrowdReplicationPositions"), bDebugReplicationPositions, TEXT("Crowd debug replication positions on server and client"), ECVF_Cheat);
+	FAutoConsoleVariableRef CVarbDebugReplication(TEXT("mass.debug.CrowdReplicationPositions"), bDebugReplicationPositions, TEXT("Crowd debug replication positions on server and client"), ECVF_Cheat);
 
 #if WITH_MASSGAMEPLAY_DEBUG && WITH_EDITOR
 	// @todo provide a better way of selecting agents to debug
diff --git a/Engine/Plugins/AI/MassCrowd/Source/MassCrowd/Private/MassCrowdServerRepresentationLODProcessor.cpp b/Engine/Plugins/AI/MassCrowd/Source/MassCrowd/Private/MassCrowdServerRepresentationLODProcessor.cpp
--- a/Engine/Plugins/AI/MassCrowd/Source/MassCrowd/Private/MassCrowdServerRepresentationLODProcessor.cpp
+++ b/Engine/Plugins/AI/MassCrowd/Source/MassCrowd/Private/MassCrowdServerRepresentationLODProcessor.cpp
@@ -10,7 +10,7 @@
 namespace UE::MassCrowd
 {
 	int32 bDebugCrowdServerRepresentationLOD = 0;
-	FAutoConsoleVariableRef CVarDebugServerRepresentationLODTest(TEXT("ai.debug.CrowdServerRepresentationLOD"), bDebugCrowdServerRepresentationLOD, TEXT("Debug Crowd ServerRepresentation LOD"), ECVF_Cheat);
+	FAutoConsoleVariableRef CVarDebugServerRepresentationLODTest(TEXT("mass.debug.CrowdServerRepresentationLOD"), bDebugCrowdServerRepresentationLOD, TEXT("Debug Crowd ServerRepresentation LOD"), ECVF_Cheat);
 } // UE::MassCrowd
 
 UMassCrowdServerRepresentationLODProcessor::UMassCrowdServerRepresentationLODProcessor()
diff --git a/Engine/Plugins/AI/MassCrowd/Source/MassCrowd/Private/MassCrowdVisualizationLODProcessor.cpp b/Engine/Plugins/AI/MassCrowd/Source/MassCrowd/Private/MassCrowdVisualizationLODProcessor.cpp
--- a/Engine/Plugins/AI/MassCrowd/Source/MassCrowd/Private/MassCrowdVisualizationLODProcessor.cpp
+++ b/Engine/Plugins/AI/MassCrowd/Source/MassCrowd/Private/MassCrowdVisualizationLODProcessor.cpp
@@ -18,8 +18,8 @@ namespace UE::MassCrowd
 
 	FAutoConsoleVariableRef ConsoleVariables[] =
 	{
-		FAutoConsoleVariableRef(TEXT("ai.debug.CrowdVisualizationLOD"), bDebugCrowdVisualizationLOD, TEXT("Debug crowd visualization LOD"), ECVF_Cheat),
-		FAutoConsoleVariableRef(TEXT("ai.debug.ShowISMUnderSpecifiedRange"), bDebugShowISMUnderSpecifiedRange, TEXT("Show ISM under a specified range (meters)"), ECVF_Cheat)
+		FAutoConsoleVariableRef(TEXT("mass.debug.crowd.VisualizationLOD"), bDebugCrowdVisualizationLOD, TEXT("Debug crowd visualization LOD"), ECVF_Cheat),
+		FAutoConsoleVariableRef(TEXT("mass.debug.crowd.ShowISMUnderSpecifiedRange"), bDebugShowISMUnderSpecifiedRange, TEXT("Show ISM under a specified range (meters)"), ECVF_Cheat)
 	};
 
 } // UE::MassCrowd
diff --git a/Engine/Plugins/AI/MassCrowd/Source/MassCrowd/Private/MassCrowdVisualizationProcessor.cpp b/Engine/Plugins/AI/MassCrowd/Source/MassCrowd/Private/MassCrowdVisualizationProcessor.cpp
--- a/Engine/Plugins/AI/MassCrowd/Source/MassCrowd/Private/MassCrowdVisualizationProcessor.cpp
+++ b/Engine/Plugins/AI/MassCrowd/Source/MassCrowd/Private/MassCrowdVisualizationProcessor.cpp
@@ -9,7 +9,7 @@
 namespace UE::MassCrowd
 {
 	int32 bDebugCrowdVisualType = 0;
-	FAutoConsoleVariableRef CVarDebugVisualType(TEXT("ai.debug.CrowdVisualType"), bDebugCrowdVisualType, TEXT("Debug Crowd Visual Type"), ECVF_Cheat);
+	FAutoConsoleVariableRef CVarDebugVisualType(TEXT("mass.debug.CrowdVisualType"), bDebugCrowdVisualType, TEXT("Debug Crowd Visual Type"), ECVF_Cheat);
 
 	FColor CrowdRepresentationTypesColors[] =
 	{
diff --git a/Engine/Plugins/Runtime/MassGameplay/Source/MassLOD/Private/MassSimulationLOD.cpp b/Engine/Plugins/Runtime/MassGameplay/Source/MassLOD/Private/MassSimulationLOD.cpp
--- a/Engine/Plugins/Runtime/MassGameplay/Source/MassLOD/Private/MassSimulationLOD.cpp
+++ b/Engine/Plugins/Runtime/MassGameplay/Source/MassLOD/Private/MassSimulationLOD.cpp
@@ -59,7 +59,7 @@ FMassSimulationVariableTickSharedFragment::FMassSimulationVariableTickSharedFrag
 namespace UE::MassLOD
 {
 	int32 bDebugSimulationLOD = 0;
-	FAutoConsoleVariableRef CVarDebugSimulationLODTest(TEXT("ai.debug.SimulationLOD"), bDebugSimulationLOD, TEXT("Debug Simulation LOD"), ECVF_Cheat);
+	FAutoConsoleVariableRef CVarDebugSimulationLODTest(TEXT("mass.debug.SimulationLOD"), bDebugSimulationLOD, TEXT("Debug Simulation LOD"), ECVF_Cheat);
 } // UE::MassLOD
 
 UMassSimulationLODProcessor::UMassSimulationLODProcessor()
diff --git a/Engine/Plugins/Runtime/MassGameplay/Source/MassReplication/Private/MassReplicationProcessor.cpp b/Engine/Plugins/Runtime/MassGameplay/Source/MassReplication/Private/MassReplicationProcessor.cpp
--- a/Engine/Plugins/Runtime/MassGameplay/Source/MassReplication/Private/MassReplicationProcessor.cpp
+++ b/Engine/Plugins/Runtime/MassGameplay/Source/MassReplication/Private/MassReplicationProcessor.cpp
@@ -9,7 +9,7 @@
 namespace UE::Mass::Replication
 {
 	int32 DebugClientReplicationLOD = -1;
-	FAutoConsoleVariableRef CVarDebugReplicationViewerLOD(TEXT("ai.debug.ClientReplicationLOD"), DebugClientReplicationLOD, TEXT("Debug Replication LOD of the specified client index"), ECVF_Cheat);
+	FAutoConsoleVariableRef CVarDebugReplicationViewerLOD(TEXT("mass.debug.ClientReplicationLOD"), DebugClientReplicationLOD, TEXT("Debug Replication LOD of the specified client index"), ECVF_Cheat);
 } // UE::Mass::Crowd
 
 //----------------------------------------------------------------------//
diff --git a/Engine/Plugins/Runtime/MassGameplay/Source/MassRepresentation/Private/MassVisualizationLODProcessor.cpp b/Engine/Plugins/Runtime/MassGameplay/Source/MassRepresentation/Private/MassVisualizationLODProcessor.cpp
--- a/Engine/Plugins/Runtime/MassGameplay/Source/MassRepresentation/Private/MassVisualizationLODProcessor.cpp
+++ b/Engine/Plugins/Runtime/MassGameplay/Source/MassRepresentation/Private/MassVisualizationLODProcessor.cpp
@@ -6,7 +6,7 @@
 namespace UE::MassRepresentation
 {
 	int32 bDebugRepresentationLOD = 0;
-	FAutoConsoleVariableRef CVarDebugRepresentationLOD(TEXT("ai.debug.RepresentationLOD"), bDebugRepresentationLOD, TEXT("Debug representation LOD"), ECVF_Cheat);
+	FAutoConsoleVariableRef CVarDebugRepresentationLOD(TEXT("mass.debug.RepresentationLOD"), bDebugRepresentationLOD, TEXT("Debug representation LOD"), ECVF_Cheat);
 } // UE::MassRepresentation
 
 
