[Mass] Compile fix for error C2039: 'DebugDisplayLOD': is not a member of 'TMassLODCalculator<FLODDefaultLogic>' in MassCrowd

#rnx

diff --git a/Engine/Plugins/AI/MassCrowd/Source/MassCrowd/Private/MassCrowdServerRepresentationLODProcessor.cpp b/Engine/Plugins/AI/MassCrowd/Source/MassCrowd/Private/MassCrowdServerRepresentationLODProcessor.cpp
--- a/Engine/Plugins/AI/MassCrowd/Source/MassCrowd/Private/MassCrowdServerRepresentationLODProcessor.cpp
+++ b/Engine/Plugins/AI/MassCrowd/Source/MassCrowd/Private/MassCrowdServerRepresentationLODProcessor.cpp
@@ -1,6 +1,7 @@
 // Copyright Epic Games, Inc. All Rights Reserved.
 
 #include "MassCrowdServerRepresentationLODProcessor.h"
+#include "MassCommonTypes.h"
 #include "MassCommonFragments.h"
 #include "MassCrowdFragments.h"
 #include "MassExecutionContext.h"
@@ -82,6 +83,7 @@ void UMassCrowdServerRepresentationLODProcessor::Execute(FMassEntityManager& Ent
 		}
 	}
 
+#if WITH_MASSGAMEPLAY_DEBUG
 	// Optional debug display
 	if (UE::MassCrowd::bDebugCrowdServerRepresentationLOD)
 	{
@@ -94,4 +96,5 @@ void UMassCrowdServerRepresentationLODProcessor::Execute(FMassEntityManager& Ent
 			LODCalculator.DebugDisplayLOD(Context, RepresentationLODFragments, LocationList, World);
 		});
 	}
+#endif // WITH_MASSGAMEPLAY_DEBUG
 }
diff --git a/Engine/Plugins/AI/MassCrowd/Source/MassCrowd/Private/MassCrowdVisualizationLODProcessor.cpp b/Engine/Plugins/AI/MassCrowd/Source/MassCrowd/Private/MassCrowdVisualizationLODProcessor.cpp
--- a/Engine/Plugins/AI/MassCrowd/Source/MassCrowd/Private/MassCrowdVisualizationLODProcessor.cpp
+++ b/Engine/Plugins/AI/MassCrowd/Source/MassCrowd/Private/MassCrowdVisualizationLODProcessor.cpp
@@ -2,6 +2,7 @@
 
 #include "MassCrowdVisualizationLODProcessor.h"
 #include "MassCommonFragments.h"
+#include "MassCommonTypes.h"
 #include "MassExecutionContext.h"
 #include "MassCrowdFragments.h"
 #include "MassEntityManager.h"
@@ -58,7 +59,8 @@ void UMassCrowdVisualizationLODProcessor::Execute(FMassEntityManager& EntityMana
 	TRACE_CPUPROFILER_EVENT_SCOPE(TEXT("CrowdVisualizationLOD"))
 
 	Super::Execute(EntityManager, Context);
-	
+
+#if WITH_MASSGAMEPLAY_DEBUG
 	if (UE::MassCrowd::bDebugCrowdVisualizationLOD)
 	{
 		TRACE_CPUPROFILER_EVENT_SCOPE(TEXT("DebugDisplayLOD"))
@@ -71,6 +73,7 @@ void UMassCrowdVisualizationLODProcessor::Execute(FMassEntityManager& EntityMana
 			LODSharedFragment.LODCalculator.DebugDisplayLOD(Context, VisualizationLODList, LocationList, World);
 		});
 	}
+#endif // WITH_MASSGAMEPLAY_DEBUG
 
 	if (UE::MassCrowd::bDebugShowISMUnderSpecifiedRange > 0)
 	{
