[Mass] Fixed a check condition in UMassCrowdDynamicObstacleDeinitializer::Initialize that erronously triggered in inactive worlds.

diff --git a/Engine/Plugins/AI/MassCrowd/Source/MassCrowd/Private/MassCrowdNavigationProcessor.cpp b/Engine/Plugins/AI/MassCrowd/Source/MassCrowd/Private/MassCrowdNavigationProcessor.cpp
--- a/Engine/Plugins/AI/MassCrowd/Source/MassCrowd/Private/MassCrowdNavigationProcessor.cpp
+++ b/Engine/Plugins/AI/MassCrowd/Source/MassCrowd/Private/MassCrowdNavigationProcessor.cpp
@@ -278,8 +278,10 @@ void UMassCrowdDynamicObstacleDeinitializer::Initialize(UObject& Owner)
 {
 	Super::Initialize(Owner);
 
-	ZoneGraphAnnotationSubsystem = UWorld::GetSubsystem<UZoneGraphAnnotationSubsystem>(Owner.GetWorld());
-	checkf(ZoneGraphAnnotationSubsystem != nullptr, TEXT("UZoneGraphAnnotationSubsystem is mandatory when using this processor."));
+	UWorld* World = Owner.GetWorld();
+	ZoneGraphAnnotationSubsystem = UWorld::GetSubsystem<UZoneGraphAnnotationSubsystem>(World);
+	checkf(ZoneGraphAnnotationSubsystem != nullptr || (World && World->WorldType == EWorldType::Inactive)
+		, TEXT("UZoneGraphAnnotationSubsystem is mandatory when using this processor."));
 }
 
 void UMassCrowdDynamicObstacleDeinitializer::ConfigureQueries()
