[MassCrowd]
* Fixed non unity compile error by adding include

diff --git a/Engine/Plugins/AI/MassCrowd/Source/MassCrowd/Private/MassCrowdRepresentationActorManagement.cpp b/Engine/Plugins/AI/MassCrowd/Source/MassCrowd/Private/MassCrowdRepresentationActorManagement.cpp
--- a/Engine/Plugins/AI/MassCrowd/Source/MassCrowd/Private/MassCrowdRepresentationActorManagement.cpp
+++ b/Engine/Plugins/AI/MassCrowd/Source/MassCrowd/Private/MassCrowdRepresentationActorManagement.cpp
@@ -1,8 +1,9 @@
 // Copyright Epic Games, Inc. All Rights Reserved.
 
 #include "MassCrowdRepresentationActorManagement.h"
- #include "Components/SkeletalMeshComponent.h"
+#include "Components/SkeletalMeshComponent.h"
 #include "Components/CapsuleComponent.h"
+#include "Engine/World.h"
 #include "GameFramework/CharacterMovementComponent.h"
 #include "GameFramework/Character.h"
 #include "MassAgentComponent.h"
