Used UnrealCodeFixup to fix dll storage on code

diff --git a/Engine/Plugins/AI/MassCrowd/Source/MassCrowd/Public/MassCrowdBubble.h b/Engine/Plugins/AI/MassCrowd/Source/MassCrowd/Public/MassCrowdBubble.h
--- a/Engine/Plugins/AI/MassCrowd/Source/MassCrowd/Public/MassCrowdBubble.h
+++ b/Engine/Plugins/AI/MassCrowd/Source/MassCrowd/Public/MassCrowdBubble.h
@@ -10,9 +10,11 @@
 #include "MassEntityView.h"
 #include "MassCrowdBubble.generated.h"
 
+#define UE_API MASSCROWD_API
+
 class FMassCrowdClientBubbleHandler;
 
-class MASSCROWD_API FMassCrowdClientBubbleHandler : public TClientBubbleHandlerBase<FCrowdFastArrayItem>
+class FMassCrowdClientBubbleHandler : public TClientBubbleHandlerBase<FCrowdFastArrayItem>
 {
 public:
 	typedef TClientBubbleHandlerBase<FCrowdFastArrayItem> Super;
@@ -35,15 +37,15 @@ public:
 
 protected:
 #if UE_REPLICATION_COMPILE_CLIENT_CODE
-	virtual void PostReplicatedAdd(const TArrayView<int32> AddedIndices, int32 FinalSize) override;
-	virtual void PostReplicatedChange(const TArrayView<int32> ChangedIndices, int32 FinalSize) override;
+	UE_API virtual void PostReplicatedAdd(const TArrayView<int32> AddedIndices, int32 FinalSize) override;
+	UE_API virtual void PostReplicatedChange(const TArrayView<int32> ChangedIndices, int32 FinalSize) override;
 
-	void PostReplicatedChangeEntity(const FMassEntityView& EntityView, const FReplicatedCrowdAgent& Item) const;
+	UE_API void PostReplicatedChangeEntity(const FMassEntityView& EntityView, const FReplicatedCrowdAgent& Item) const;
 #endif //UE_REPLICATION_COMPILE_CLIENT_CODE
 
 #if WITH_MASSGAMEPLAY_DEBUG && WITH_EDITOR
-	virtual void DebugValidateBubbleOnServer() override;
-	virtual void DebugValidateBubbleOnClient() override;
+	UE_API virtual void DebugValidateBubbleOnServer() override;
+	UE_API virtual void DebugValidateBubbleOnClient() override;
 #endif // WITH_MASSGAMEPLAY_DEBUG
 
 	FMassClientBubblePathHandler PathHandler;
@@ -52,7 +54,7 @@ protected:
 
 /** Mass client bubble, there will be one of these per client and it will handle replicating the fast array of Agents between the server and clients */
 USTRUCT()
-struct MASSCROWD_API FMassCrowdClientBubbleSerializer : public FMassClientBubbleSerializerBase
+struct FMassCrowdClientBubbleSerializer : public FMassClientBubbleSerializerBase
 {
 	GENERATED_BODY()
 
@@ -91,20 +93,22 @@ struct TStructOpsTypeTraits<FMassCrowdClientBubbleSerializer> : public TStructOp
  *  This class will allow us to replicate Mass data based on the fidelity required for each player controller. There is one AMassReplicationActor per PlayerController and 
  *  which is also its owner.
  */
-UCLASS()
-class MASSCROWD_API AMassCrowdClientBubbleInfo : public AMassClientBubbleInfoBase
+UCLASS(MinimalAPI)
+class AMassCrowdClientBubbleInfo : public AMassClientBubbleInfoBase
 {
 	GENERATED_BODY()
 
 public:
-	AMassCrowdClientBubbleInfo(const FObjectInitializer& ObjectInitializer);
+	UE_API AMassCrowdClientBubbleInfo(const FObjectInitializer& ObjectInitializer);
 
 	FMassCrowdClientBubbleSerializer& GetCrowdSerializer() { return CrowdSerializer; }
 
 protected:
-	virtual void GetLifetimeReplicatedProps(TArray<FLifetimeProperty>& OutLifetimeProps) const override;
+	UE_API virtual void GetLifetimeReplicatedProps(TArray<FLifetimeProperty>& OutLifetimeProps) const override;
 
 protected:
 	UPROPERTY(Replicated, Transient)
 	FMassCrowdClientBubbleSerializer CrowdSerializer;
 };
+
+#undef UE_API
diff --git a/Engine/Plugins/AI/MassCrowd/Source/MassCrowd/Public/MassCrowdFragments.h b/Engine/Plugins/AI/MassCrowd/Source/MassCrowd/Public/MassCrowdFragments.h
--- a/Engine/Plugins/AI/MassCrowd/Source/MassCrowd/Public/MassCrowdFragments.h
+++ b/Engine/Plugins/AI/MassCrowd/Source/MassCrowd/Public/MassCrowdFragments.h
@@ -12,7 +12,7 @@
  * Should not contain any data, this is purely a tag
  */
 USTRUCT()
-struct MASSCROWD_API FMassCrowdTag : public FMassTag
+struct FMassCrowdTag : public FMassTag
 {
 	GENERATED_BODY()
 };
@@ -21,7 +21,7 @@ struct MASSCROWD_API FMassCrowdTag : public FMassTag
  * Data fragment to store the last lane the agent was tracked on.
  */
 USTRUCT()
-struct MASSCROWD_API FMassCrowdLaneTrackingFragment : public FMassFragment
+struct FMassCrowdLaneTrackingFragment : public FMassFragment
 {
 	GENERATED_BODY()
 	FZoneGraphLaneHandle TrackedLaneHandle;
@@ -29,7 +29,7 @@ struct MASSCROWD_API FMassCrowdLaneTrackingFragment : public FMassFragment
 
 
 USTRUCT()
-struct MASSCROWD_API FMassCrowdObstacleFragment : public FMassFragment
+struct FMassCrowdObstacleFragment : public FMassFragment
 {
 	GENERATED_BODY()
 
diff --git a/Engine/Plugins/AI/MassCrowd/Source/MassCrowd/Public/MassCrowdLaneDataRenderingComponent.h b/Engine/Plugins/AI/MassCrowd/Source/MassCrowd/Public/MassCrowdLaneDataRenderingComponent.h
--- a/Engine/Plugins/AI/MassCrowd/Source/MassCrowd/Public/MassCrowdLaneDataRenderingComponent.h
+++ b/Engine/Plugins/AI/MassCrowd/Source/MassCrowd/Public/MassCrowdLaneDataRenderingComponent.h
@@ -5,12 +5,14 @@
 #include "Components/PrimitiveComponent.h"
 #include "MassCrowdLaneDataRenderingComponent.generated.h"
 
+#define UE_API MASSCROWD_API
+
 /**
  * Primitive component that can be used to render runtime state of zone graph lanes (e.g. Opened|Closed, Density, etc.)
  * The component must be added on a ZoneGraphData actor.
  */
-UCLASS(editinlinenew, meta = (BlueprintSpawnableComponent), hidecategories = (Object, LOD, Lighting, VirtualTexture, Transform, HLOD, Collision, TextureStreaming, Mobile, Physics, Tags, AssetUserData, Activation, Cooking, Rendering, Navigation))
-class MASSCROWD_API UMassCrowdLaneDataRenderingComponent : public UPrimitiveComponent
+UCLASS(MinimalAPI, editinlinenew, meta = (BlueprintSpawnableComponent), hidecategories = (Object, LOD, Lighting, VirtualTexture, Transform, HLOD, Collision, TextureStreaming, Mobile, Physics, Tags, AssetUserData, Activation, Cooking, Rendering, Navigation))
+class UMassCrowdLaneDataRenderingComponent : public UPrimitiveComponent
 {
 	GENERATED_BODY()
 
@@ -19,15 +21,17 @@ public:
 
 #if !UE_BUILD_SHIPPING && !UE_BUILD_TEST
 private:
-	virtual FPrimitiveSceneProxy* CreateSceneProxy() override;
-	virtual void OnRegister() override;
-	virtual void OnUnregister() override;
-	virtual FBoxSphereBounds CalcBounds(const FTransform& LocalToWorld) const override;
-	void DebugDrawOnCanvas(UCanvas* Canvas, APlayerController*) const;
+	UE_API virtual FPrimitiveSceneProxy* CreateSceneProxy() override;
+	UE_API virtual void OnRegister() override;
+	UE_API virtual void OnUnregister() override;
+	UE_API virtual FBoxSphereBounds CalcBounds(const FTransform& LocalToWorld) const override;
+	UE_API void DebugDrawOnCanvas(UCanvas* Canvas, APlayerController*) const;
 	FDelegateHandle DebugTextDrawingDelegateHandle;
 	FDelegateHandle OnLaneStateChangedDelegateHandle;
 #if WITH_EDITOR
 	FDelegateHandle OnLaneRenderSettingsChangedDelegateHandle;
 #endif // WITH_EDITOR
 #endif // !UE_BUILD_SHIPPING && !UE_BUILD_TEST
-};
\ No newline at end of file
+};
+
+#undef UE_API
diff --git a/Engine/Plugins/AI/MassCrowd/Source/MassCrowd/Public/MassCrowdMemberTrait.h b/Engine/Plugins/AI/MassCrowd/Source/MassCrowd/Public/MassCrowdMemberTrait.h
--- a/Engine/Plugins/AI/MassCrowd/Source/MassCrowd/Public/MassCrowdMemberTrait.h
+++ b/Engine/Plugins/AI/MassCrowd/Source/MassCrowd/Public/MassCrowdMemberTrait.h
@@ -5,14 +5,18 @@
 #include "MassEntityTraitBase.h"
 #include "MassCrowdMemberTrait.generated.h"
 
+#define UE_API MASSCROWD_API
+
 /**
  * Trait to mark an entity with the crowd tag and add required fragments to track current lane
  */
-UCLASS(meta = (DisplayName = "CrowdMember"))
-class MASSCROWD_API UMassCrowdMemberTrait : public UMassEntityTraitBase
+UCLASS(MinimalAPI, meta = (DisplayName = "CrowdMember"))
+class UMassCrowdMemberTrait : public UMassEntityTraitBase
 {
 	GENERATED_BODY()
 
 protected:
-	virtual void BuildTemplate(FMassEntityTemplateBuildContext& BuildContext, const UWorld& World) const override;
+	UE_API virtual void BuildTemplate(FMassEntityTemplateBuildContext& BuildContext, const UWorld& World) const override;
 };
+
+#undef UE_API
diff --git a/Engine/Plugins/AI/MassCrowd/Source/MassCrowd/Public/MassCrowdNavigationProcessor.h b/Engine/Plugins/AI/MassCrowd/Source/MassCrowd/Public/MassCrowdNavigationProcessor.h
--- a/Engine/Plugins/AI/MassCrowd/Source/MassCrowd/Public/MassCrowdNavigationProcessor.h
+++ b/Engine/Plugins/AI/MassCrowd/Source/MassCrowd/Public/MassCrowdNavigationProcessor.h
@@ -7,55 +7,57 @@
 #include "MassCrowdFragments.h"
 #include "MassCrowdNavigationProcessor.generated.h"
 
+#define UE_API MASSCROWD_API
+
 class UZoneGraphAnnotationSubsystem;
 class UMassCrowdSubsystem;
 
 /** Processor that monitors when entities change lane and notify the MassCrowd subsystem. */
-UCLASS()
-class MASSCROWD_API UMassCrowdLaneTrackingSignalProcessor : public UMassSignalProcessorBase
+UCLASS(MinimalAPI)
+class UMassCrowdLaneTrackingSignalProcessor : public UMassSignalProcessorBase
 {
 	GENERATED_BODY()
 public:
-	UMassCrowdLaneTrackingSignalProcessor();
+	UE_API UMassCrowdLaneTrackingSignalProcessor();
 
 protected:
-	virtual void InitializeInternal(UObject& Owner, const TSharedRef<FMassEntityManager>& EntityManager) override;
-	virtual void ConfigureQueries(const TSharedRef<FMassEntityManager>& EntityManager) override;
-	virtual void SignalEntities(FMassEntityManager& EntityManager, FMassExecutionContext& Context, FMassSignalNameLookup& /*Unused*/) override;
+	UE_API virtual void InitializeInternal(UObject& Owner, const TSharedRef<FMassEntityManager>& EntityManager) override;
+	UE_API virtual void ConfigureQueries(const TSharedRef<FMassEntityManager>& EntityManager) override;
+	UE_API virtual void SignalEntities(FMassEntityManager& EntityManager, FMassExecutionContext& Context, FMassSignalNameLookup& /*Unused*/) override;
 };
 
 /** Processors that cleans up the lane tracking on entity destruction. */
-UCLASS()
-class MASSCROWD_API UMassCrowdLaneTrackingDestructor : public UMassObserverProcessor
+UCLASS(MinimalAPI)
+class UMassCrowdLaneTrackingDestructor : public UMassObserverProcessor
 {
 	GENERATED_BODY()
 
 public:
-	UMassCrowdLaneTrackingDestructor();
+	UE_API UMassCrowdLaneTrackingDestructor();
 
 protected:
-	virtual void ConfigureQueries(const TSharedRef<FMassEntityManager>& EntityManager) override;
-	virtual void Execute(FMassEntityManager& EntityManager, FMassExecutionContext& Context) override;
+	UE_API virtual void ConfigureQueries(const TSharedRef<FMassEntityManager>& EntityManager) override;
+	UE_API virtual void Execute(FMassEntityManager& EntityManager, FMassExecutionContext& Context) override;
 
 	FMassEntityQuery EntityQuery;
 };
 
 
-UCLASS()
-class MASSCROWD_API UMassCrowdDynamicObstacleProcessor : public UMassProcessor
+UCLASS(MinimalAPI)
+class UMassCrowdDynamicObstacleProcessor : public UMassProcessor
 {
 	GENERATED_BODY()
 
 public:
-	UMassCrowdDynamicObstacleProcessor();
+	UE_API UMassCrowdDynamicObstacleProcessor();
 
 protected:
-	virtual void InitializeInternal(UObject& Owner, const TSharedRef<FMassEntityManager>& EntityManager)override;
-	virtual void ConfigureQueries(const TSharedRef<FMassEntityManager>& EntityManager) override;
-	virtual void Execute(FMassEntityManager& EntityManager, FMassExecutionContext& Context);
+	UE_API virtual void InitializeInternal(UObject& Owner, const TSharedRef<FMassEntityManager>& EntityManager)override;
+	UE_API virtual void ConfigureQueries(const TSharedRef<FMassEntityManager>& EntityManager) override;
+	UE_API virtual void Execute(FMassEntityManager& EntityManager, FMassExecutionContext& Context);
 
-	void OnStop(FMassCrowdObstacleFragment& OutObstacle, const float Radius);
-	void OnMove(FMassCrowdObstacleFragment& OutObstacle);
+	UE_API void OnStop(FMassCrowdObstacleFragment& OutObstacle, const float Radius);
+	UE_API void OnMove(FMassCrowdObstacleFragment& OutObstacle);
 
 	FMassEntityQuery EntityQuery_Conditional;
 
@@ -64,37 +66,39 @@ protected:
 };
 
 
-UCLASS()
-class MASSCROWD_API UMassCrowdDynamicObstacleInitializer : public UMassObserverProcessor
+UCLASS(MinimalAPI)
+class UMassCrowdDynamicObstacleInitializer : public UMassObserverProcessor
 {
 	GENERATED_BODY()
 
 public:
-	UMassCrowdDynamicObstacleInitializer();
+	UE_API UMassCrowdDynamicObstacleInitializer();
 
 protected:
-	virtual void ConfigureQueries(const TSharedRef<FMassEntityManager>& EntityManager) override;
-	virtual void Execute(FMassEntityManager& EntityManager, FMassExecutionContext& Context) override;
+	UE_API virtual void ConfigureQueries(const TSharedRef<FMassEntityManager>& EntityManager) override;
+	UE_API virtual void Execute(FMassEntityManager& EntityManager, FMassExecutionContext& Context) override;
 
 	FMassEntityQuery EntityQuery;
 };
 
 
-UCLASS()
-class MASSCROWD_API UMassCrowdDynamicObstacleDeinitializer : public UMassObserverProcessor
+UCLASS(MinimalAPI)
+class UMassCrowdDynamicObstacleDeinitializer : public UMassObserverProcessor
 {
 	GENERATED_BODY()
 
 public:
-	UMassCrowdDynamicObstacleDeinitializer();
+	UE_API UMassCrowdDynamicObstacleDeinitializer();
 
 protected:
-	virtual void InitializeInternal(UObject& Owner, const TSharedRef<FMassEntityManager>& EntityManager) override;
-	virtual void ConfigureQueries(const TSharedRef<FMassEntityManager>& EntityManager) override;
-	virtual void Execute(FMassEntityManager& EntityManager, FMassExecutionContext& Context) override;
+	UE_API virtual void InitializeInternal(UObject& Owner, const TSharedRef<FMassEntityManager>& EntityManager) override;
+	UE_API virtual void ConfigureQueries(const TSharedRef<FMassEntityManager>& EntityManager) override;
+	UE_API virtual void Execute(FMassEntityManager& EntityManager, FMassExecutionContext& Context) override;
 
 	UPROPERTY(Transient)
 	TObjectPtr<UZoneGraphAnnotationSubsystem> ZoneGraphAnnotationSubsystem;
 
 	FMassEntityQuery EntityQuery;
 };
+
+#undef UE_API
diff --git a/Engine/Plugins/AI/MassCrowd/Source/MassCrowd/Public/MassCrowdReplicatedAgent.h b/Engine/Plugins/AI/MassCrowd/Source/MassCrowd/Public/MassCrowdReplicatedAgent.h
--- a/Engine/Plugins/AI/MassCrowd/Source/MassCrowd/Public/MassCrowdReplicatedAgent.h
+++ b/Engine/Plugins/AI/MassCrowd/Source/MassCrowd/Public/MassCrowdReplicatedAgent.h
@@ -10,7 +10,7 @@
 
 /** The data that is replicated specific to each Crowd agent */
 USTRUCT()
-struct MASSCROWD_API FReplicatedCrowdAgent : public FReplicatedAgentBase
+struct FReplicatedCrowdAgent : public FReplicatedAgentBase
 {
 	GENERATED_BODY()
 
@@ -34,7 +34,7 @@ private:
 
 /** Fast array item for efficient agent replication. Remember to make this dirty if any FReplicatedCrowdAgent member variables are modified */
 USTRUCT()
-struct MASSCROWD_API FCrowdFastArrayItem : public FMassFastArrayItemBase
+struct FCrowdFastArrayItem : public FMassFastArrayItemBase
 {
 	GENERATED_BODY()
 
@@ -49,4 +49,4 @@ struct MASSCROWD_API FCrowdFastArrayItem : public FMassFastArrayItemBase
 
 	UPROPERTY()
 	FReplicatedCrowdAgent Agent;
-};
\ No newline at end of file
+};
diff --git a/Engine/Plugins/AI/MassCrowd/Source/MassCrowd/Public/MassCrowdReplicator.h b/Engine/Plugins/AI/MassCrowd/Source/MassCrowd/Public/MassCrowdReplicator.h
--- a/Engine/Plugins/AI/MassCrowd/Source/MassCrowd/Public/MassCrowdReplicator.h
+++ b/Engine/Plugins/AI/MassCrowd/Source/MassCrowd/Public/MassCrowdReplicator.h
@@ -5,9 +5,11 @@
 #include "MassReplicationProcessor.h"
 #include "MassCrowdReplicator.generated.h"
 
+#define UE_API MASSCROWD_API
+
 /** Class that handles replication and only runs on the server. It queries Mass entity fragments and sets those values when appropriate using the MassClientBubbleHandler. */
-UCLASS()
-class MASSCROWD_API UMassCrowdReplicator : public UMassReplicatorBase
+UCLASS(MinimalAPI)
+class UMassCrowdReplicator : public UMassReplicatorBase
 {
 	GENERATED_BODY()
 
@@ -16,11 +18,13 @@ public:
 	 * Overridden to add specific entity query requirements for replication.
 	 * Usually we add replication processor handler requirements.
 	 */
-	virtual void AddRequirements(FMassEntityQuery& EntityQuery) override;
+	UE_API virtual void AddRequirements(FMassEntityQuery& EntityQuery) override;
 	
 	/**
 	 * Overridden to process the client replication.
 	 * This methods should call CalculateClientReplication with the appropriate callback implementation.
 	 */
-	virtual void ProcessClientReplication(FMassExecutionContext& Context, FMassReplicationContext& ReplicationContext) override;
-};
\ No newline at end of file
+	UE_API virtual void ProcessClientReplication(FMassExecutionContext& Context, FMassReplicationContext& ReplicationContext) override;
+};
+
+#undef UE_API
diff --git a/Engine/Plugins/AI/MassCrowd/Source/MassCrowd/Public/MassCrowdRepresentationActorManagement.h b/Engine/Plugins/AI/MassCrowd/Source/MassCrowd/Public/MassCrowdRepresentationActorManagement.h
--- a/Engine/Plugins/AI/MassCrowd/Source/MassCrowd/Public/MassCrowdRepresentationActorManagement.h
+++ b/Engine/Plugins/AI/MassCrowd/Source/MassCrowd/Public/MassCrowdRepresentationActorManagement.h
@@ -5,12 +5,14 @@
 #include "MassRepresentationActorManagement.h"
 #include "MassCrowdRepresentationActorManagement.generated.h"
 
+#define UE_API MASSCROWD_API
+
 /**
  * Overridden representation processor to make it tied to the crowd via the requirements.
  * It is also the base class for all the different type of crowd representation (Visualization & ServerSideRepresentation)
  */
-UCLASS(abstract)
-class MASSCROWD_API UMassCrowdRepresentationActorManagement : public UMassRepresentationActorManagement
+UCLASS(MinimalAPI, abstract)
+class UMassCrowdRepresentationActorManagement : public UMassRepresentationActorManagement
 {
 	GENERATED_BODY()
 
@@ -23,7 +25,7 @@ protected:
 	 * @param EntityIdx is the entity index currently processing
 	 * @param Context is the current Mass execution context
 	 */
-	virtual void SetActorEnabled(const EMassActorEnabledType EnabledType, AActor& Actor, const int32 EntityIdx, FMassCommandBuffer& CommandBuffer) const override;
+	UE_API virtual void SetActorEnabled(const EMassActorEnabledType EnabledType, AActor& Actor, const int32 EntityIdx, FMassCommandBuffer& CommandBuffer) const override;
 
 	/**
 	 * Returns an actor of the template type and setup fragments values from it
@@ -37,7 +39,7 @@ protected:
 	 * @param Priority of this spawn request in comparison with the others, lower value means higher priority
 	 * @return the actor spawned
 	 */
-	virtual AActor* GetOrSpawnActor(UMassRepresentationSubsystem& RepresentationSubsystem, FMassEntityManager& EntityManager
+	UE_API virtual AActor* GetOrSpawnActor(UMassRepresentationSubsystem& RepresentationSubsystem, FMassEntityManager& EntityManager
 		, const FMassEntityHandle MassAgent, const FTransform& Transform, const int16 TemplateActorIndex
 		, FMassActorSpawnRequestHandle& SpawnRequestHandle, const float Priority) const override;
 
@@ -48,5 +50,7 @@ protected:
 	 * @param Actor is the actual actor to teleport
 	 * @param Context is the current Mass execution context
 	 */
-	virtual void TeleportActor(const FTransform& Transform, AActor& Actor, FMassCommandBuffer& CommandBuffer) const override;
-};
\ No newline at end of file
+	UE_API virtual void TeleportActor(const FTransform& Transform, AActor& Actor, FMassCommandBuffer& CommandBuffer) const override;
+};
+
+#undef UE_API
diff --git a/Engine/Plugins/AI/MassCrowd/Source/MassCrowd/Public/MassCrowdRepresentationSubsystem.h b/Engine/Plugins/AI/MassCrowd/Source/MassCrowd/Public/MassCrowdRepresentationSubsystem.h
--- a/Engine/Plugins/AI/MassCrowd/Source/MassCrowd/Public/MassCrowdRepresentationSubsystem.h
+++ b/Engine/Plugins/AI/MassCrowd/Source/MassCrowd/Public/MassCrowdRepresentationSubsystem.h
@@ -5,16 +5,20 @@
 #include "MassRepresentationSubsystem.h"
 #include "MassCrowdRepresentationSubsystem.generated.h"
 
+#define UE_API MASSCROWD_API
+
 /**
  * Subsystem responsible for all visual of mass crowd agents, will handle actors spawning and static mesh instances
  */
-UCLASS()
-class MASSCROWD_API UMassCrowdRepresentationSubsystem : public UMassRepresentationSubsystem
+UCLASS(MinimalAPI)
+class UMassCrowdRepresentationSubsystem : public UMassRepresentationSubsystem
 {
 	GENERATED_BODY()
 
 protected:
 	// USubsystem BEGIN
-	virtual void Initialize(FSubsystemCollectionBase& Collection) override;
+	UE_API virtual void Initialize(FSubsystemCollectionBase& Collection) override;
 	// USubsystem END
-};
\ No newline at end of file
+};
+
+#undef UE_API
diff --git a/Engine/Plugins/AI/MassCrowd/Source/MassCrowd/Public/MassCrowdServerRepresentationLODProcessor.h b/Engine/Plugins/AI/MassCrowd/Source/MassCrowd/Public/MassCrowdServerRepresentationLODProcessor.h
--- a/Engine/Plugins/AI/MassCrowd/Source/MassCrowd/Public/MassCrowdServerRepresentationLODProcessor.h
+++ b/Engine/Plugins/AI/MassCrowd/Source/MassCrowd/Public/MassCrowdServerRepresentationLODProcessor.h
@@ -12,19 +12,19 @@
 
 #include "MassCrowdServerRepresentationLODProcessor.generated.h"
 
-UCLASS(meta=(DisplayName="Crowd Simulation LOD"))
-class MASSCROWD_API UMassCrowdServerRepresentationLODProcessor : public UMassProcessor
+UCLASS(MinimalAPI, meta=(DisplayName="Crowd Simulation LOD"))
+class UMassCrowdServerRepresentationLODProcessor : public UMassProcessor
 {
 	GENERATED_BODY()
 
 public:
-	UMassCrowdServerRepresentationLODProcessor();
+	MASSCROWD_API UMassCrowdServerRepresentationLODProcessor();
 
 protected:
 
-	virtual void ConfigureQueries(const TSharedRef<FMassEntityManager>& EntityManager) override;
-	virtual void InitializeInternal(UObject& InOwner, const TSharedRef<FMassEntityManager>& EntityManager) override;
-	virtual void Execute(FMassEntityManager& EntityManager, FMassExecutionContext& Context) override;
+	MASSCROWD_API virtual void ConfigureQueries(const TSharedRef<FMassEntityManager>& EntityManager) override;
+	MASSCROWD_API virtual void InitializeInternal(UObject& InOwner, const TSharedRef<FMassEntityManager>& EntityManager) override;
+	MASSCROWD_API virtual void Execute(FMassEntityManager& EntityManager, FMassExecutionContext& Context) override;
 
 	/** Distance where each LOD becomes relevant */
 	UPROPERTY(EditAnywhere, Category = "Mass|LOD", config)
diff --git a/Engine/Plugins/AI/MassCrowd/Source/MassCrowd/Public/MassCrowdServerRepresentationProcessor.h b/Engine/Plugins/AI/MassCrowd/Source/MassCrowd/Public/MassCrowdServerRepresentationProcessor.h
--- a/Engine/Plugins/AI/MassCrowd/Source/MassCrowd/Public/MassCrowdServerRepresentationProcessor.h
+++ b/Engine/Plugins/AI/MassCrowd/Source/MassCrowd/Public/MassCrowdServerRepresentationProcessor.h
@@ -5,18 +5,22 @@
 #include "MassRepresentationProcessor.h"
 #include "MassCrowdServerRepresentationProcessor.generated.h"
 
+#define UE_API MASSCROWD_API
+
 /**
  * Overridden representation processor to make it tied to the crowd on the server via the requirements
  * It is the counter part of the crowd visualization processor on the client.
  */
-UCLASS(meta = (DisplayName = "Mass Crowd Server Representation"))
-class MASSCROWD_API UMassCrowdServerRepresentationProcessor : public UMassRepresentationProcessor
+UCLASS(MinimalAPI, meta = (DisplayName = "Mass Crowd Server Representation"))
+class UMassCrowdServerRepresentationProcessor : public UMassRepresentationProcessor
 {
 	GENERATED_BODY()
 
 public:
-	UMassCrowdServerRepresentationProcessor();
+	UE_API UMassCrowdServerRepresentationProcessor();
 
 	/** Configure the owned FMassEntityQuery instances to express processor's requirements */
-	virtual void ConfigureQueries(const TSharedRef<FMassEntityManager>& EntityManager) override;
-};
\ No newline at end of file
+	UE_API virtual void ConfigureQueries(const TSharedRef<FMassEntityManager>& EntityManager) override;
+};
+
+#undef UE_API
diff --git a/Engine/Plugins/AI/MassCrowd/Source/MassCrowd/Public/MassCrowdServerRepresentationTrait.h b/Engine/Plugins/AI/MassCrowd/Source/MassCrowd/Public/MassCrowdServerRepresentationTrait.h
--- a/Engine/Plugins/AI/MassCrowd/Source/MassCrowd/Public/MassCrowdServerRepresentationTrait.h
+++ b/Engine/Plugins/AI/MassCrowd/Source/MassCrowd/Public/MassCrowdServerRepresentationTrait.h
@@ -11,16 +11,16 @@
 #include "MassCrowdServerRepresentationTrait.generated.h"
 
 
-UCLASS(meta=(DisplayName="Crowd Server Representation"))
-class MASSCROWD_API UMassCrowdServerRepresentationTrait : public UMassEntityTraitBase
+UCLASS(MinimalAPI, meta=(DisplayName="Crowd Server Representation"))
+class UMassCrowdServerRepresentationTrait : public UMassEntityTraitBase
 {
 	GENERATED_BODY()
 
-	UMassCrowdServerRepresentationTrait();
+	MASSCROWD_API UMassCrowdServerRepresentationTrait();
 
 protected:
 
-	virtual void BuildTemplate(FMassEntityTemplateBuildContext& BuildContext, const UWorld& World) const override;
+	MASSCROWD_API virtual void BuildTemplate(FMassEntityTemplateBuildContext& BuildContext, const UWorld& World) const override;
 
 	/** Actor class of this agent when spawned on server */
 	UPROPERTY(EditAnywhere, Category = "Mass|Visual")
diff --git a/Engine/Plugins/AI/MassCrowd/Source/MassCrowd/Public/MassCrowdSettings.h b/Engine/Plugins/AI/MassCrowd/Source/MassCrowd/Public/MassCrowdSettings.h
--- a/Engine/Plugins/AI/MassCrowd/Source/MassCrowd/Public/MassCrowdSettings.h
+++ b/Engine/Plugins/AI/MassCrowd/Source/MassCrowd/Public/MassCrowdSettings.h
@@ -6,6 +6,8 @@
 #include "ZoneGraphTypes.h"
 #include "MassCrowdSettings.generated.h"
 
+#define UE_API MASSCROWD_API
+
 #if WITH_EDITOR
 /** Called when density settings change. */
 DECLARE_MULTICAST_DELEGATE(FOnMassCrowdLaneDataSettingsChanged);
@@ -19,7 +21,7 @@ DECLARE_MULTICAST_DELEGATE(FOnMassCrowdLaneRenderSettingsChanged);
  * to maintain overall density during the simulation
  */
 USTRUCT()
-struct MASSCROWD_API FMassCrowdLaneDensityDesc
+struct FMassCrowdLaneDensityDesc
 {
 	GENERATED_BODY()
 
@@ -46,8 +48,8 @@ struct MASSCROWD_API FMassCrowdLaneDensityDesc
 /**
  * Settings for the MassCrowd plugin.
  */
-UCLASS(config = Plugins, defaultconfig, DisplayName = "Mass Crowd", dontCollapseCategories)
-class MASSCROWD_API UMassCrowdSettings : public UMassModuleSettings
+UCLASS(MinimalAPI, config = Plugins, defaultconfig, DisplayName = "Mass Crowd", dontCollapseCategories)
+class UMassCrowdSettings : public UMassModuleSettings
 {
 	GENERATED_BODY()
 
@@ -134,7 +136,7 @@ public:
 protected:
 
 #if WITH_EDITOR
-	virtual void PostEditChangeChainProperty(FPropertyChangedChainEvent& PropertyChangedEvent) override;
+	UE_API virtual void PostEditChangeChainProperty(FPropertyChangedChainEvent& PropertyChangedEvent) override;
 #endif
 
 	/** Distance ahead of the current lane location for the next movement target location. */
@@ -148,4 +150,6 @@ protected:
 	/** List of all lane density descriptors. */
 	UPROPERTY(EditAnywhere, config, Category = "Lane Density")
 	TArray<FMassCrowdLaneDensityDesc> LaneDensities;
-};
\ No newline at end of file
+};
+
+#undef UE_API
diff --git a/Engine/Plugins/AI/MassCrowd/Source/MassCrowd/Public/MassCrowdSpawnerSubsystem.h b/Engine/Plugins/AI/MassCrowd/Source/MassCrowd/Public/MassCrowdSpawnerSubsystem.h
--- a/Engine/Plugins/AI/MassCrowd/Source/MassCrowd/Public/MassCrowdSpawnerSubsystem.h
+++ b/Engine/Plugins/AI/MassCrowd/Source/MassCrowd/Public/MassCrowdSpawnerSubsystem.h
@@ -8,8 +8,8 @@
 /**
  * A subsystem managing spawning of crowd
  */
-UCLASS()
-class MASSCROWD_API UMassCrowdSpawnerSubsystem : public UMassActorSpawnerSubsystem
+UCLASS(MinimalAPI)
+class UMassCrowdSpawnerSubsystem : public UMassActorSpawnerSubsystem
 {
 	GENERATED_BODY()
 };
diff --git a/Engine/Plugins/AI/MassCrowd/Source/MassCrowd/Public/MassCrowdSubsystem.h b/Engine/Plugins/AI/MassCrowd/Source/MassCrowd/Public/MassCrowdSubsystem.h
--- a/Engine/Plugins/AI/MassCrowd/Source/MassCrowd/Public/MassCrowdSubsystem.h
+++ b/Engine/Plugins/AI/MassCrowd/Source/MassCrowd/Public/MassCrowdSubsystem.h
@@ -62,8 +62,8 @@ struct FRegisteredCrowdLaneData
  * Subsystem that tracks mass entities that are wandering on the zone graph.
  * It will create custom runtime lane data to allow branching decisions.
  */
-UCLASS()
-class MASSCROWD_API UMassCrowdSubsystem : public UMassSubsystemBase
+UCLASS(MinimalAPI)
+class UMassCrowdSubsystem : public UMassSubsystemBase
 {
 	GENERATED_BODY()
 public:
@@ -73,11 +73,11 @@ public:
 
 #if WITH_EDITOR
 	/** Clears and rebuilds all lane and intersection data for registered zone graphs using the current settings. */
-	void RebuildLaneData();
+	MASSCROWD_API void RebuildLaneData();
 #endif
 
 	/** @return true if the Crowd subsystem has lane data for specified graph. */
-	bool HasCrowdDataForZoneGraph(const FZoneGraphDataHandle DataHandle) const;
+	MASSCROWD_API bool HasCrowdDataForZoneGraph(const FZoneGraphDataHandle DataHandle) const;
 
 	/**
 	 * Returns the readonly runtime data associated to a given zone graph.
@@ -85,42 +85,42 @@ public:
 	 * @return Runtime data associated to the zone graph if available; nullptr otherwise
 	 * @note Method will ensure if DataHandle is invalid or if associated data doesn't exist. Should call HasCrowdDataForZoneGraph first.
 	 */
-	const FRegisteredCrowdLaneData* GetCrowdData(const FZoneGraphDataHandle DataHandle) const;
+	MASSCROWD_API const FRegisteredCrowdLaneData* GetCrowdData(const FZoneGraphDataHandle DataHandle) const;
 
 	/**
 	 * Returns the readonly runtime data associated to a given zone graph lane.
 	 * @param LaneHandle A valid lane handle used to retrieve the runtime data; ensure if handle is invalid
 	 * @return Runtime data associated to the lane (if available)
 	 */
-	TOptional<FZoneGraphCrowdLaneData> GetCrowdLaneData(const FZoneGraphLaneHandle LaneHandle) const;
+	MASSCROWD_API TOptional<FZoneGraphCrowdLaneData> GetCrowdLaneData(const FZoneGraphLaneHandle LaneHandle) const;
 
 	/**
 	 * Returns the entity tracking runtime data associated to a given zone graph lane.
 	 * @param LaneHandle A valid lane handle used to retrieve the associated tracking data; ensure if handle is invalid
 	 * @return Runtime data associated to the lane (nullptr if provided handle is invalid or no data is associated to that lane)
 	 */
-	const FCrowdTrackingLaneData* GetCrowdTrackingLaneData(const FZoneGraphLaneHandle LaneHandle) const;
+	MASSCROWD_API const FCrowdTrackingLaneData* GetCrowdTrackingLaneData(const FZoneGraphLaneHandle LaneHandle) const;
 
 	/**
 	 * Returns the branching data associated to a given zone graph lane.
  	 * @param LaneHandle A valid lane handle used to retrieve the associated data; ensure if handle is invalid
 	 * @return Branching data associated to the lane (nullptr if provided handle is invalid or no data is associated to that lane)
 	 */
-	const FCrowdBranchingLaneData* GetCrowdBranchingLaneData(const FZoneGraphLaneHandle LaneHandle) const;
+	MASSCROWD_API const FCrowdBranchingLaneData* GetCrowdBranchingLaneData(const FZoneGraphLaneHandle LaneHandle) const;
 
 	/**
 	 * Returns the waiting area runtime data associated to a given zone graph lane.
 	 * @param LaneHandle A valid lane handle used to retrieve the associated intersection data; ensure if handle is invalid
 	 * @return Runtime data associated to the lane (nullptr if provided handle is invalid or no data is associated to that lane)
 	 */
-	const FCrowdWaitAreaData* GetCrowdWaitingAreaData(const FZoneGraphLaneHandle LaneHandle) const;
+	MASSCROWD_API const FCrowdWaitAreaData* GetCrowdWaitingAreaData(const FZoneGraphLaneHandle LaneHandle) const;
 
 	/**
 	 * Return the current state of a lane.
 	 * @param LaneHandle A valid lane handle used to retrieve the runtime data and change the lane state; ensure if handle is invalid
 	 * @return The state of the lane
 	 */
-	ECrowdLaneState GetLaneState(const FZoneGraphLaneHandle LaneHandle) const;
+	MASSCROWD_API ECrowdLaneState GetLaneState(const FZoneGraphLaneHandle LaneHandle) const;
 
 	/**
 	 * Changes the state of a lane.
@@ -128,7 +128,7 @@ public:
 	 * @param NewState The new state of the lane
 	 * @return True if the state was successfully changed or was already in the right state; false if handle is invalid
 	 */
-	bool SetLaneState(const FZoneGraphLaneHandle LaneHandle, ECrowdLaneState NewState);
+	MASSCROWD_API bool SetLaneState(const FZoneGraphLaneHandle LaneHandle, ECrowdLaneState NewState);
 
 	/** @return Combined Tag mask that represents all possible lane density tags. Built from MassCrowdSettings. */
 	FZoneGraphTagMask GetDensityMask() const { return DensityMask; }
@@ -142,7 +142,7 @@ public:
 	 * @param OutSlotDirection Facing direction associated to the acquire slot.
 	 * @return Index of the slot, or INDEX_NONE if no slots are available or if the lane is not a waiting lane.
 	 */
-	int32 AcquireWaitingSlot(const FMassEntityHandle Entity, const FVector& EntityPosition, const FZoneGraphLaneHandle LaneHandle,
+	MASSCROWD_API int32 AcquireWaitingSlot(const FMassEntityHandle Entity, const FVector& EntityPosition, const FZoneGraphLaneHandle LaneHandle,
 							 FVector& OutSlotPosition, FVector& OutSlotDirection);
 	
 	/**
@@ -151,7 +151,7 @@ public:
 	 * @param LaneHandle A handle to a lane with waiting data.
 	 * @param SlotIndex Index of the previously acquired slot.
 	 */
-	void ReleaseWaitingSlot(const FMassEntityHandle Entity, const FZoneGraphLaneHandle LaneHandle, const int32 SlotIndex);
+	MASSCROWD_API void ReleaseWaitingSlot(const FMassEntityHandle Entity, const FZoneGraphLaneHandle LaneHandle, const int32 SlotIndex);
 
 	/**
 	 * Callback from the lane tracker processor to indicates a mass entity changing lane.
@@ -159,40 +159,40 @@ public:
 	 * @param PreviousLaneHandle Last frame lane handle (can be invalid)
 	 * @param CurrentLaneHandle Current frame lane handle (can be invalid)
 	 */
-	void OnEntityLaneChanged(const FMassEntityHandle Entity, const FZoneGraphLaneHandle PreviousLaneHandle, const FZoneGraphLaneHandle CurrentLaneHandle);
+	MASSCROWD_API void OnEntityLaneChanged(const FMassEntityHandle Entity, const FZoneGraphLaneHandle PreviousLaneHandle, const FZoneGraphLaneHandle CurrentLaneHandle);
 
 	/** Returns the weight for lane selection that is associated to the given lane based on its density tag. */
-	float GetDensityWeight(const FZoneGraphLaneHandle LaneHandle, const FZoneGraphTagMask LaneTagMask) const;
+	MASSCROWD_API float GetDensityWeight(const FZoneGraphLaneHandle LaneHandle, const FZoneGraphTagMask LaneTagMask) const;
 
 protected:
 	friend class UMassCrowdLaneTrackingProcessor;
 	friend class UMassCrowdWanderFragmentDestructor;
 	friend class UZoneGraphCrowdLaneAnnotations;
 
-	void UpdateDensityMask();
+	MASSCROWD_API void UpdateDensityMask();
 
 	// USubsystem BEGIN
-	virtual void Initialize(FSubsystemCollectionBase& Collection) override;
-	virtual void PostInitialize() override;
-	virtual void Deinitialize() override;
+	MASSCROWD_API virtual void Initialize(FSubsystemCollectionBase& Collection) override;
+	MASSCROWD_API virtual void PostInitialize() override;
+	MASSCROWD_API virtual void Deinitialize() override;
 	// USubsystem END
 
-	void PostZoneGraphDataAdded(const AZoneGraphData* ZoneGraphData);
-	void PreZoneGraphDataRemoved(const AZoneGraphData* ZoneGraphData);
+	MASSCROWD_API void PostZoneGraphDataAdded(const AZoneGraphData* ZoneGraphData);
+	MASSCROWD_API void PreZoneGraphDataRemoved(const AZoneGraphData* ZoneGraphData);
 
 	/**
 	 * Returns the modifiable runtime data associated to a given zone graph lane.
 	 * @param LaneHandle A valid lane handle used to retrieve the runtime data; ensure if handle is invalid
 	 * @return Runtime data associated to the lane if available; nullptr otherwise
 	 */
-	FZoneGraphCrowdLaneData* GetMutableCrowdLaneData(const FZoneGraphLaneHandle LaneHandle);
+	MASSCROWD_API FZoneGraphCrowdLaneData* GetMutableCrowdLaneData(const FZoneGraphLaneHandle LaneHandle);
 
 	/**
 	 * Populates the crowd lane data array with all existing lane in the zone graph.
 	 * Also creates some extra data for lanes that could be used to wait in case of
 	 * temporarily closed lane (e.g. intersection crossings).
 	 */
-	void BuildLaneData(FRegisteredCrowdLaneData& LaneData, const FZoneGraphStorage& Storage);
+	MASSCROWD_API void BuildLaneData(FRegisteredCrowdLaneData& LaneData, const FZoneGraphStorage& Storage);
 
 	/**
 	 * Callback to keep count of entities currently on a given lane.
@@ -200,7 +200,7 @@ protected:
 	 * @param LaneIndex Index of the lane
 	 * @param TrackingData Runtime state associated to the lane
 	 */
-	void OnEnterTrackedLane(const FMassEntityHandle Entity, const int32 LaneIndex, FCrowdTrackingLaneData& TrackingData);
+	MASSCROWD_API void OnEnterTrackedLane(const FMassEntityHandle Entity, const int32 LaneIndex, FCrowdTrackingLaneData& TrackingData);
 
 	/**
 	 * Callback to keep count of entities currently on a given lane.
@@ -208,7 +208,7 @@ protected:
 	 * @param LaneIndex Index of the lane
 	 * @param TrackingData Runtime state associated to the lane
 	 */
-	void OnExitTrackedLane(const FMassEntityHandle Entity, const int32 LaneIndex, FCrowdTrackingLaneData& TrackingData);
+	MASSCROWD_API void OnExitTrackedLane(const FMassEntityHandle Entity, const int32 LaneIndex, FCrowdTrackingLaneData& TrackingData);
 
 	/**
 	 * Creates and initializes the occupancy data of a lane.
@@ -216,7 +216,7 @@ protected:
 	 * @param ZoneGraphStorage ZoneGraph data storage to extract lane information from
 	 * @return The newly created structure holding occupancy data associated to the zone graph lane
 	 */
-	FCrowdTrackingLaneData& CreateTrackingData(const int32 LaneIndex, const FZoneGraphStorage& ZoneGraphStorage);
+	MASSCROWD_API FCrowdTrackingLaneData& CreateTrackingData(const int32 LaneIndex, const FZoneGraphStorage& ZoneGraphStorage);
 
 	/**
 	 * Creates dedicated structure to hold data of a branching lane.
@@ -224,9 +224,9 @@ protected:
 	 * @param ZoneGraphStorage ZoneGraph data storage to extract lane information from
 	 * @return The newly created structure holding specific data of a branching lane
 	 */
-	FCrowdBranchingLaneData& CreateBranchingData(const int32 LaneIndex, const FZoneGraphStorage& ZoneGraphStorage);
+	MASSCROWD_API FCrowdBranchingLaneData& CreateBranchingData(const int32 LaneIndex, const FZoneGraphStorage& ZoneGraphStorage);
 
-	void CreateWaitSlots(const int32 CrossingLaneIndex, FCrowdWaitAreaData& WaitArea, const FZoneGraphStorage& ZoneGraphStorage);
+	MASSCROWD_API void CreateWaitSlots(const int32 CrossingLaneIndex, FCrowdWaitAreaData& WaitArea, const FZoneGraphStorage& ZoneGraphStorage);
 
 	UPROPERTY(Transient)
 	TObjectPtr<UZoneGraphSubsystem> ZoneGraphSubsystem = nullptr;
diff --git a/Engine/Plugins/AI/MassCrowd/Source/MassCrowd/Public/MassCrowdTypes.h b/Engine/Plugins/AI/MassCrowd/Source/MassCrowd/Public/MassCrowdTypes.h
--- a/Engine/Plugins/AI/MassCrowd/Source/MassCrowd/Public/MassCrowdTypes.h
+++ b/Engine/Plugins/AI/MassCrowd/Source/MassCrowd/Public/MassCrowdTypes.h
@@ -45,7 +45,7 @@ enum class ECrowdLaneState : uint8
 
 /** Runtime data associated to lane that can be used to track count of entities on it. */
 USTRUCT()
-struct MASSCROWD_API FCrowdTrackingLaneData
+struct FCrowdTrackingLaneData
 {
 	GENERATED_BODY()
 
@@ -60,7 +60,7 @@ struct MASSCROWD_API FCrowdTrackingLaneData
 
 /** Runtime data associated to lane that can be used to wait another one to open. */
 USTRUCT()
-struct MASSCROWD_API FCrowdWaitSlot
+struct FCrowdWaitSlot
 {
 	GENERATED_BODY()
 
@@ -75,7 +75,7 @@ struct MASSCROWD_API FCrowdWaitSlot
 
 /** Runtime data associated to entry to a lane that can be opened or closed. */
 USTRUCT()
-struct MASSCROWD_API FCrowdWaitAreaData
+struct FCrowdWaitAreaData
 {
 	GENERATED_BODY()
 
@@ -99,7 +99,7 @@ struct MASSCROWD_API FCrowdWaitAreaData
 
 /** Runtime data associated to branching. */
 USTRUCT()
-struct MASSCROWD_API FCrowdBranchingLaneData
+struct FCrowdBranchingLaneData
 {
 	GENERATED_BODY()
 
diff --git a/Engine/Plugins/AI/MassCrowd/Source/MassCrowd/Public/MassCrowdVisualizationLODProcessor.h b/Engine/Plugins/AI/MassCrowd/Source/MassCrowd/Public/MassCrowdVisualizationLODProcessor.h
--- a/Engine/Plugins/AI/MassCrowd/Source/MassCrowd/Public/MassCrowdVisualizationLODProcessor.h
+++ b/Engine/Plugins/AI/MassCrowd/Source/MassCrowd/Public/MassCrowdVisualizationLODProcessor.h
@@ -6,32 +6,35 @@
 #include "MassLODCollectorProcessor.h"
 #include "MassCrowdVisualizationLODProcessor.generated.h"
 
+#define UE_API MASSCROWD_API
+
 /*
  * Created a crowd version for parallelization of the crowd with the traffic
  */
-UCLASS(meta=(DisplayName="Crowd visualization LOD"))
-class MASSCROWD_API UMassCrowdVisualizationLODProcessor : public UMassVisualizationLODProcessor
+UCLASS(MinimalAPI, meta=(DisplayName="Crowd visualization LOD"))
+class UMassCrowdVisualizationLODProcessor : public UMassVisualizationLODProcessor
 {
 	GENERATED_BODY()
 public:
-	UMassCrowdVisualizationLODProcessor();
+	UE_API UMassCrowdVisualizationLODProcessor();
 
 protected:
-	virtual void ConfigureQueries(const TSharedRef<FMassEntityManager>& EntityManager) override;
-	virtual void Execute(FMassEntityManager& EntityManager, FMassExecutionContext& Context) override;
+	UE_API virtual void ConfigureQueries(const TSharedRef<FMassEntityManager>& EntityManager) override;
+	UE_API virtual void Execute(FMassEntityManager& EntityManager, FMassExecutionContext& Context) override;
 };
 
 /*
  * Created a crowd version for parallelization of the crowd with the traffic
  */
-UCLASS(meta = (DisplayName = "Crowd LOD Collection "))
-class MASSCROWD_API UMassCrowdLODCollectorProcessor : public UMassLODCollectorProcessor
+UCLASS(MinimalAPI, meta = (DisplayName = "Crowd LOD Collection "))
+class UMassCrowdLODCollectorProcessor : public UMassLODCollectorProcessor
 {
 	GENERATED_BODY()
 
-	UMassCrowdLODCollectorProcessor();
+	UE_API UMassCrowdLODCollectorProcessor();
 
 protected:
-	virtual void ConfigureQueries(const TSharedRef<FMassEntityManager>& EntityManager) override;
+	UE_API virtual void ConfigureQueries(const TSharedRef<FMassEntityManager>& EntityManager) override;
 };
 
+#undef UE_API
diff --git a/Engine/Plugins/AI/MassCrowd/Source/MassCrowd/Public/MassCrowdVisualizationProcessor.h b/Engine/Plugins/AI/MassCrowd/Source/MassCrowd/Public/MassCrowdVisualizationProcessor.h
--- a/Engine/Plugins/AI/MassCrowd/Source/MassCrowd/Public/MassCrowdVisualizationProcessor.h
+++ b/Engine/Plugins/AI/MassCrowd/Source/MassCrowd/Public/MassCrowdVisualizationProcessor.h
@@ -5,48 +5,52 @@
 #include "MassRepresentationProcessor.h"
 #include "MassCrowdVisualizationProcessor.generated.h"
 
+#define UE_API MASSCROWD_API
+
 /**
  * Overridden visualization processor to make it tied to the crowd via the requirements
  */
-UCLASS(meta = (DisplayName = "Mass Crowd Visualization"))
-class MASSCROWD_API UMassCrowdVisualizationProcessor : public UMassVisualizationProcessor
+UCLASS(MinimalAPI, meta = (DisplayName = "Mass Crowd Visualization"))
+class UMassCrowdVisualizationProcessor : public UMassVisualizationProcessor
 {
 	GENERATED_BODY()
 public:
-	UMassCrowdVisualizationProcessor();
+	UE_API UMassCrowdVisualizationProcessor();
 
 protected:
 
 	/** Configure the owned FMassEntityQuery instances to express processor's requirements */
-	virtual void ConfigureQueries(const TSharedRef<FMassEntityManager>& EntityManager) override;
+	UE_API virtual void ConfigureQueries(const TSharedRef<FMassEntityManager>& EntityManager) override;
 };
 
 /**
  * A custom visualization processor for debugging mass crowd
  */
-UCLASS(meta = (DisplayName = "Mass Crowd Visualization"))
-class MASSCROWD_API UMassDebugCrowdVisualizationProcessor : public UMassProcessor
+UCLASS(MinimalAPI, meta = (DisplayName = "Mass Crowd Visualization"))
+class UMassDebugCrowdVisualizationProcessor : public UMassProcessor
 {
 	GENERATED_BODY()
 public:
-	UMassDebugCrowdVisualizationProcessor();
+	UE_API UMassDebugCrowdVisualizationProcessor();
 
 protected:
 
 	/** Configure the owned FMassEntityQuery instances to express processor's requirements */
-	virtual void ConfigureQueries(const TSharedRef<FMassEntityManager>& EntityManager) override;
+	UE_API virtual void ConfigureQueries(const TSharedRef<FMassEntityManager>& EntityManager) override;
 
-	virtual void InitializeInternal(UObject& Owner, const TSharedRef<FMassEntityManager>& EntityManager) override;
+	UE_API virtual void InitializeInternal(UObject& Owner, const TSharedRef<FMassEntityManager>& EntityManager) override;
 
 	/**
 	 * Execution method for this processor
 	 * @param EntitySubsystem is the system to execute the lambdas on each entity chunk
 	 * @param Context is the execution context to be passed when executing the lambdas
 	 */
-	virtual void Execute(FMassEntityManager& EntityManager, FMassExecutionContext& Context) override;
+	UE_API virtual void Execute(FMassEntityManager& EntityManager, FMassExecutionContext& Context) override;
 
 	UPROPERTY(Transient)
 	TObjectPtr<UWorld> World;
 
 	FMassEntityQuery EntityQuery;
-};
\ No newline at end of file
+};
+
+#undef UE_API
diff --git a/Engine/Plugins/AI/MassCrowd/Source/MassCrowd/Public/MassCrowdVisualizationTrait.h b/Engine/Plugins/AI/MassCrowd/Source/MassCrowd/Public/MassCrowdVisualizationTrait.h
--- a/Engine/Plugins/AI/MassCrowd/Source/MassCrowd/Public/MassCrowdVisualizationTrait.h
+++ b/Engine/Plugins/AI/MassCrowd/Source/MassCrowd/Public/MassCrowdVisualizationTrait.h
@@ -5,12 +5,16 @@
 #include "MassVisualizationTrait.h"
 #include "MassCrowdVisualizationTrait.generated.h"
 
-UCLASS(BlueprintType, EditInlineNew, CollapseCategories, meta=(DisplayName="Crowd Visualization"))
-class MASSCROWD_API UMassCrowdVisualizationTrait : public UMassVisualizationTrait
+#define UE_API MASSCROWD_API
+
+UCLASS(MinimalAPI, BlueprintType, EditInlineNew, CollapseCategories, meta=(DisplayName="Crowd Visualization"))
+class UMassCrowdVisualizationTrait : public UMassVisualizationTrait
 {
 	GENERATED_BODY()
 public:
-	UMassCrowdVisualizationTrait();
+	UE_API UMassCrowdVisualizationTrait();
 
-	virtual void BuildTemplate(FMassEntityTemplateBuildContext& BuildContext, const UWorld& World) const override;
+	UE_API virtual void BuildTemplate(FMassEntityTemplateBuildContext& BuildContext, const UWorld& World) const override;
 };
+
+#undef UE_API
diff --git a/Engine/Plugins/AI/MassCrowd/Source/MassCrowd/Public/Tasks/MassCrowdClaimWaitSlotTask.h b/Engine/Plugins/AI/MassCrowd/Source/MassCrowd/Public/Tasks/MassCrowdClaimWaitSlotTask.h
--- a/Engine/Plugins/AI/MassCrowd/Source/MassCrowd/Public/Tasks/MassCrowdClaimWaitSlotTask.h
+++ b/Engine/Plugins/AI/MassCrowd/Source/MassCrowd/Public/Tasks/MassCrowdClaimWaitSlotTask.h
@@ -7,6 +7,8 @@
 #include "Tasks/MassZoneGraphPathFollowTask.h"
 #include "MassCrowdClaimWaitSlotTask.generated.h"
 
+#define UE_API MASSCROWD_API
+
 struct FStateTreeExecutionContext;
 class UMassCrowdSubsystem;
 namespace UE::MassBehavior
@@ -18,7 +20,7 @@ namespace UE::MassBehavior
 * Claim wait slot and expose slot position for path follow.
 */
 USTRUCT()
-struct MASSCROWD_API FMassCrowdClaimWaitSlotTaskInstanceData
+struct FMassCrowdClaimWaitSlotTaskInstanceData
 {
 	GENERATED_BODY()
 
@@ -33,22 +35,24 @@ struct MASSCROWD_API FMassCrowdClaimWaitSlotTaskInstanceData
 };
 
 USTRUCT(meta = (DisplayName = "Crowd Claim Wait Slot"))
-struct MASSCROWD_API FMassCrowdClaimWaitSlotTask : public FMassStateTreeTaskBase
+struct FMassCrowdClaimWaitSlotTask : public FMassStateTreeTaskBase
 {
 	GENERATED_BODY()
 
 	using FInstanceDataType = FMassCrowdClaimWaitSlotTaskInstanceData;
 	
-	FMassCrowdClaimWaitSlotTask();
+	UE_API FMassCrowdClaimWaitSlotTask();
 
 protected:
-	virtual bool Link(FStateTreeLinker& Linker) override;
+	UE_API virtual bool Link(FStateTreeLinker& Linker) override;
 	virtual const UStruct* GetInstanceDataType() const override { return FInstanceDataType::StaticStruct(); }
-	virtual EStateTreeRunStatus EnterState(FStateTreeExecutionContext& Context, const FStateTreeTransitionResult& Transition) const override;
-	virtual void ExitState(FStateTreeExecutionContext& Context, const FStateTreeTransitionResult& Transition) const override;
-	virtual void GetDependencies(UE::MassBehavior::FStateTreeDependencyBuilder& Builder) const override;
+	UE_API virtual EStateTreeRunStatus EnterState(FStateTreeExecutionContext& Context, const FStateTreeTransitionResult& Transition) const override;
+	UE_API virtual void ExitState(FStateTreeExecutionContext& Context, const FStateTreeTransitionResult& Transition) const override;
+	UE_API virtual void GetDependencies(UE::MassBehavior::FStateTreeDependencyBuilder& Builder) const override;
 
 	TStateTreeExternalDataHandle<FMassZoneGraphLaneLocationFragment> LocationHandle;
 	TStateTreeExternalDataHandle<FMassMoveTargetFragment> MoveTargetHandle;
 	TStateTreeExternalDataHandle<UMassCrowdSubsystem> CrowdSubsystemHandle;
 };
+
+#undef UE_API
diff --git a/Engine/Plugins/AI/MassCrowd/Source/MassCrowd/Public/Tasks/MassZoneGraphFindWanderTarget.h b/Engine/Plugins/AI/MassCrowd/Source/MassCrowd/Public/Tasks/MassZoneGraphFindWanderTarget.h
--- a/Engine/Plugins/AI/MassCrowd/Source/MassCrowd/Public/Tasks/MassZoneGraphFindWanderTarget.h
+++ b/Engine/Plugins/AI/MassCrowd/Source/MassCrowd/Public/Tasks/MassZoneGraphFindWanderTarget.h
@@ -6,6 +6,8 @@
 #include "Tasks/MassZoneGraphPathFollowTask.h"
 #include "MassZoneGraphFindWanderTarget.generated.h"
 
+#define UE_API MASSCROWD_API
+
 struct FStateTreeExecutionContext;
 struct FMassZoneGraphLaneLocationFragment;
 class UZoneGraphSubsystem;
@@ -20,7 +22,7 @@ namespace UE::MassBehavior
  * Updates TargetLocation to a wander target based on the agents current location on ZoneGraph.
  */
 USTRUCT()
-struct MASSCROWD_API FMassZoneGraphFindWanderTargetInstanceData
+struct FMassZoneGraphFindWanderTargetInstanceData
 {
 	GENERATED_BODY()
 
@@ -29,19 +31,19 @@ struct MASSCROWD_API FMassZoneGraphFindWanderTargetInstanceData
 };
 
 USTRUCT(meta = (DisplayName = "ZG Find Wander Target"))
-struct MASSCROWD_API FMassZoneGraphFindWanderTarget : public FMassStateTreeTaskBase
+struct FMassZoneGraphFindWanderTarget : public FMassStateTreeTaskBase
 {
 	GENERATED_BODY()
 
 	using FInstanceDataType = FMassZoneGraphFindWanderTargetInstanceData;
 	
-	FMassZoneGraphFindWanderTarget();
+	UE_API FMassZoneGraphFindWanderTarget();
 
 protected:
-	virtual bool Link(FStateTreeLinker& Linker) override;
+	UE_API virtual bool Link(FStateTreeLinker& Linker) override;
 	virtual const UStruct* GetInstanceDataType() const override { return FInstanceDataType::StaticStruct(); }
-	virtual EStateTreeRunStatus EnterState(FStateTreeExecutionContext& Context, const FStateTreeTransitionResult& Transition) const override;
-	virtual void GetDependencies(UE::MassBehavior::FStateTreeDependencyBuilder& Builder) const override;
+	UE_API virtual EStateTreeRunStatus EnterState(FStateTreeExecutionContext& Context, const FStateTreeTransitionResult& Transition) const override;
+	UE_API virtual void GetDependencies(UE::MassBehavior::FStateTreeDependencyBuilder& Builder) const override;
 
 	TStateTreeExternalDataHandle<FMassZoneGraphLaneLocationFragment> LocationHandle;
 	TStateTreeExternalDataHandle<UZoneGraphSubsystem> ZoneGraphSubsystemHandle;
@@ -51,3 +53,5 @@ protected:
 	UPROPERTY(EditAnywhere, Category = Parameter)
 	FZoneGraphTagFilter AllowedAnnotationTags;
 };
+
+#undef UE_API
diff --git a/Engine/Plugins/AI/MassCrowd/Source/MassCrowd/Public/ZoneGraphCrowdLaneAnnotations.h b/Engine/Plugins/AI/MassCrowd/Source/MassCrowd/Public/ZoneGraphCrowdLaneAnnotations.h
--- a/Engine/Plugins/AI/MassCrowd/Source/MassCrowd/Public/ZoneGraphCrowdLaneAnnotations.h
+++ b/Engine/Plugins/AI/MassCrowd/Source/MassCrowd/Public/ZoneGraphCrowdLaneAnnotations.h
@@ -8,9 +8,11 @@
 
 #include "ZoneGraphCrowdLaneAnnotations.generated.h"
 
+#define UE_API MASSCROWD_API
+
 /** Event indicating the new state of a lane. */
 USTRUCT()
-struct MASSCROWD_API FZoneGraphCrowdLaneStateChangeEvent : public FZoneGraphAnnotationEventBase
+struct FZoneGraphCrowdLaneStateChangeEvent : public FZoneGraphAnnotationEventBase
 {
 	GENERATED_BODY()
 
@@ -30,20 +32,20 @@ struct MASSCROWD_API FZoneGraphCrowdLaneStateChangeEvent : public FZoneGraphAnno
 /**
  * Zone graph blocking behavior
  */
-UCLASS(ClassGroup = AI, BlueprintType, meta = (BlueprintSpawnableComponent))
-class MASSCROWD_API UZoneGraphCrowdLaneAnnotations : public UZoneGraphAnnotationComponent
+UCLASS(MinimalAPI, ClassGroup = AI, BlueprintType, meta = (BlueprintSpawnableComponent))
+class UZoneGraphCrowdLaneAnnotations : public UZoneGraphAnnotationComponent
 {
 	GENERATED_BODY()
 
 protected:
-	virtual void PostSubsystemsInitialized() override;
-	virtual FZoneGraphTagMask GetAnnotationTags() const override;
-	virtual void HandleEvents(const FInstancedStructContainer& Events) override;
-	virtual void TickAnnotation(const float DeltaTime, FZoneGraphAnnotationTagContainer& AnnotationTagContainer) override;
+	UE_API virtual void PostSubsystemsInitialized() override;
+	UE_API virtual FZoneGraphTagMask GetAnnotationTags() const override;
+	UE_API virtual void HandleEvents(const FInstancedStructContainer& Events) override;
+	UE_API virtual void TickAnnotation(const float DeltaTime, FZoneGraphAnnotationTagContainer& AnnotationTagContainer) override;
 
 #if UE_ENABLE_DEBUG_DRAWING
-	virtual void DebugDraw(FZoneGraphAnnotationSceneProxy* DebugProxy) override;
-	virtual void DebugDrawCanvas(UCanvas* Canvas, APlayerController*) override;
+	UE_API virtual void DebugDraw(FZoneGraphAnnotationSceneProxy* DebugProxy) override;
+	UE_API virtual void DebugDrawCanvas(UCanvas* Canvas, APlayerController*) override;
 #endif // UE_ENABLE_DEBUG_DRAWING
 
 	/** Annotation Tag to mark a closed lane. */
@@ -64,3 +66,5 @@ protected:
 	UPROPERTY(Transient)
 	TObjectPtr<UMassCrowdSubsystem> CrowdSubsystem = nullptr;
 };
+
+#undef UE_API
diff --git a/Engine/Plugins/Runtime/Harmonix/Source/HarmonixMetasound/Public/HarmonixMetasound/Components/MusicClockDriverBase.h b/Engine/Plugins/Runtime/Harmonix/Source/HarmonixMetasound/Public/HarmonixMetasound/Components/MusicClockDriverBase.h
--- a/Engine/Plugins/Runtime/Harmonix/Source/HarmonixMetasound/Public/HarmonixMetasound/Components/MusicClockDriverBase.h
+++ b/Engine/Plugins/Runtime/Harmonix/Source/HarmonixMetasound/Public/HarmonixMetasound/Components/MusicClockDriverBase.h
@@ -1,4 +1,4 @@
-﻿// Copyright Epic Games, Inc. All Rights Reserved.
+// Copyright Epic Games, Inc. All Rights Reserved.
 #pragma once
 
 #include "HarmonixMetasound/Components/MusicClockComponent.h"
@@ -12,26 +12,26 @@ struct ISongMapEvaluator;
 
 #define UE_API HARMONIXMETASOUND_API
 
-class UE_API FMusicClockDriverBase : public TSharedFromThis<FMusicClockDriverBase>
+class FMusicClockDriverBase : public TSharedFromThis<FMusicClockDriverBase>
 {
 public:
-	FMusicClockDriverBase(UObject* WorldContextObject, const FMusicClockSettingsBase& Settings);
-	virtual ~FMusicClockDriverBase();
+	UE_API FMusicClockDriverBase(UObject* WorldContextObject, const FMusicClockSettingsBase& Settings);
+	UE_API virtual ~FMusicClockDriverBase();
 
-	void EnsureClockIsValidForGameFrame();
+	UE_API void EnsureClockIsValidForGameFrame();
 
 	virtual bool CalculateSongPosWithOffset(float MsOffset, ECalibratedMusicTimebase Timebase, FMidiSongPos& OutResult) const = 0;
 
-	double GetWallClockTime() const;
+	UE_API double GetWallClockTime() const;
 
 	EMusicClockState GetState() const { return MusicClockState; }
-	const FMidiSongPos& GetCurrentSongPos(ECalibratedMusicTimebase Timebase) const;
-	const FMidiSongPos& GetPreviousSongPos(ECalibratedMusicTimebase Timebase) const;
+	UE_API const FMidiSongPos& GetCurrentSongPos(ECalibratedMusicTimebase Timebase) const;
+	UE_API const FMidiSongPos& GetPreviousSongPos(ECalibratedMusicTimebase Timebase) const;
 
-	void Start();
-	void Pause();
-	void Continue();
-	void Stop();
+	UE_API void Start();
+	UE_API void Pause();
+	UE_API void Continue();
+	UE_API void Stop();
 	
 	virtual const ISongMapEvaluator* GetCurrentSongMapEvaluator() const = 0;
 
@@ -40,10 +40,10 @@ public:
 
 	virtual void Disconnect() = 0;
 
-	void UpdateMusicPlaybackRate(float InTempo, float InClockAdvanceRate, int32 InTimeSigNum, int32 InTimeSigDenom);
+	UE_API void UpdateMusicPlaybackRate(float InTempo, float InClockAdvanceRate, int32 InTimeSigNum, int32 InTimeSigDenom);
 
 	// This returns our BeatsPerMinute which is our true beats, not quarter notes.
-	float GetBeatsPerMinute() const;
+	UE_API float GetBeatsPerMinute() const;
 	
 	FSongMaps DefaultMaps;
 	
@@ -86,9 +86,9 @@ protected:
 	
 private:
 
-	void ResetDefaultSongMaps(float DefaultTempo, int32 DefaultTimeSigNum, int32 DefaultTimeSigDenom);
+	UE_API void ResetDefaultSongMaps(float DefaultTempo, int32 DefaultTimeSigNum, int32 DefaultTimeSigDenom);
 	virtual bool RefreshCurrentSongPos() = 0;
-	void UpdateBeatsAndBarsPerSecond();
+	UE_API void UpdateBeatsAndBarsPerSecond();
 
 	uint64 LastUpdateFrame = 0;
 
@@ -97,4 +97,4 @@ private:
 	TWeakObjectPtr<UObject> ContextObj;
 };
 
-#undef UE_API
\ No newline at end of file
+#undef UE_API
diff --git a/Engine/Plugins/VirtualProduction/CustomDetailsView/Source/CustomDetailsView/Public/CustomDetailsViewSequencer.h b/Engine/Plugins/VirtualProduction/CustomDetailsView/Source/CustomDetailsView/Public/CustomDetailsViewSequencer.h
--- a/Engine/Plugins/VirtualProduction/CustomDetailsView/Source/CustomDetailsView/Public/CustomDetailsViewSequencer.h
+++ b/Engine/Plugins/VirtualProduction/CustomDetailsView/Source/CustomDetailsView/Public/CustomDetailsViewSequencer.h
@@ -7,17 +7,21 @@
 #include "Delegates/DelegateCombinations.h"
 #include "Templates/SharedPointer.h"
 
+#define UE_API CUSTOMDETAILSVIEW_API
+
 class IDetailKeyframeHandler;
 class IPropertyHandle;
 struct FPropertyRowExtensionButton;
 
-struct CUSTOMDETAILSVIEW_API FCustomDetailsViewSequencerUtils
+struct FCustomDetailsViewSequencerUtils
 {
 	DECLARE_DELEGATE_RetVal(TSharedPtr<IDetailKeyframeHandler>, FGetKeyframeHandlerDelegate);
 
-	static void CreateSequencerExtensionButton(const FGetKeyframeHandlerDelegate& InKeyframeHandlerDelegate, TSharedPtr<IPropertyHandle> InPropertyHandle,
+	static UE_API void CreateSequencerExtensionButton(const FGetKeyframeHandlerDelegate& InKeyframeHandlerDelegate, TSharedPtr<IPropertyHandle> InPropertyHandle,
 		TArray<FPropertyRowExtensionButton>& OutExtensionButtons);
 
-	static void CreateSequencerExtensionButton(TWeakPtr<IDetailKeyframeHandler> InKeyframeHandlerWeak, TSharedPtr<IPropertyHandle> InPropertyHandle,
+	static UE_API void CreateSequencerExtensionButton(TWeakPtr<IDetailKeyframeHandler> InKeyframeHandlerWeak, TSharedPtr<IPropertyHandle> InPropertyHandle,
 		TArray<FPropertyRowExtensionButton>& OutExtensionButtons);
 };
+
+#undef UE_API
diff --git a/Engine/Plugins/VirtualProduction/DynamicMaterial/Source/DynamicMaterialTextureSet/Public/DMTextureSet.h b/Engine/Plugins/VirtualProduction/DynamicMaterial/Source/DynamicMaterialTextureSet/Public/DMTextureSet.h
--- a/Engine/Plugins/VirtualProduction/DynamicMaterial/Source/DynamicMaterialTextureSet/Public/DMTextureSet.h
+++ b/Engine/Plugins/VirtualProduction/DynamicMaterial/Source/DynamicMaterialTextureSet/Public/DMTextureSet.h
@@ -10,15 +10,17 @@
 
 #include "DMTextureSet.generated.h"
 
+#define UE_API DYNAMICMATERIALTEXTURESET_API
+
 enum class EDMTextureSetMaterialProperty : uint8;
 
-UCLASS(BlueprintType, ClassGroup = "Material Designer")
-class DYNAMICMATERIALTEXTURESET_API UDMTextureSet : public UObject
+UCLASS(MinimalAPI, BlueprintType, ClassGroup = "Material Designer")
+class UDMTextureSet : public UObject
 {
 	GENERATED_BODY()
 
 public:
-	UDMTextureSet();
+	UE_API UDMTextureSet();
 
 	virtual ~UDMTextureSet() override = default;
 
@@ -29,12 +31,12 @@ public:
 	 * @return True if the property exists in the Texture Map.
 	 */
 	UFUNCTION(BlueprintPure, Category = "Material Designer")
-	bool HasMaterialProperty(EDMTextureSetMaterialProperty InMaterialProperty) const;
+	UE_API bool HasMaterialProperty(EDMTextureSetMaterialProperty InMaterialProperty) const;
 
 	/**
 	 * @return Gets the entire Texture Map.
 	 */
-	const TMap<EDMTextureSetMaterialProperty, FDMMaterialTexture>& GetTextures() const;
+	UE_API const TMap<EDMTextureSetMaterialProperty, FDMMaterialTexture>& GetTextures() const;
 
 	/**
 	 * Checks whether a given Material Property has a Texture assigned to it.
@@ -42,7 +44,7 @@ public:
 	 * @return True if the Material Property has an assigned Texture.
 	 */
 	UFUNCTION(BlueprintPure, Category = "Material Designer")
-	bool HasMaterialTexture(EDMTextureSetMaterialProperty InMaterialProperty) const;
+	UE_API bool HasMaterialTexture(EDMTextureSetMaterialProperty InMaterialProperty) const;
 
 	/**
 	 * Gets the Material Texture associated with a Material Property. Does not check whether a Texture
@@ -52,7 +54,7 @@ public:
 	 * @return True if the Material Property exists within the Texture Map.
 	 */
 	UFUNCTION(BlueprintPure, Category = "Material Designer")
-	bool GetMaterialTexture(EDMTextureSetMaterialProperty InMaterialProperty, FDMMaterialTexture& OutMaterialTexture) const;
+	UE_API bool GetMaterialTexture(EDMTextureSetMaterialProperty InMaterialProperty, FDMMaterialTexture& OutMaterialTexture) const;
 
 	/**
 	 * Gets the Material Texture associated with a Material Property. Does not check whether a Texture
@@ -60,7 +62,7 @@ public:
 	 * @param InMaterialProperty The Material Property to check.
 	 * @return A pointer to the found Material Texture, if it exists, or nullptr.
 	 */
-	const FDMMaterialTexture* GetMaterialTexture(EDMTextureSetMaterialProperty InMaterialProperty) const;
+	UE_API const FDMMaterialTexture* GetMaterialTexture(EDMTextureSetMaterialProperty InMaterialProperty) const;
 
 	/**
 	 * Sets the Material Texture for a given Material Property. Can be used to unset Textures.
@@ -68,7 +70,7 @@ public:
 	 * @param InMaterialTexture The Material Texture to set on the given Material Property.
 	 */
 	UFUNCTION(BlueprintCallable, Category = "Material Designer")
-	void SetMaterialTexture(EDMTextureSetMaterialProperty InMaterialProperty, const FDMMaterialTexture& InMaterialTexture);
+	UE_API void SetMaterialTexture(EDMTextureSetMaterialProperty InMaterialProperty, const FDMMaterialTexture& InMaterialTexture);
 
 	/**
 	 * Checks the Texture Map to see if a given Texture exists within it.
@@ -76,9 +78,11 @@ public:
 	 * @return True if the Texture exits in the Texture Map.
 	 */
 	UFUNCTION(BlueprintPure, Category = "Material Designer")
-	bool ContainsTexture(UTexture* InTexture) const;
+	UE_API bool ContainsTexture(UTexture* InTexture) const;
 
 protected:
 	UPROPERTY(EditAnywhere, BlueprintReadWrite, Category = "Material Designer", EditFixedSize, meta = (ReadOnlyKeys, AllowPrivateAccess))
 	TMap<EDMTextureSetMaterialProperty, FDMMaterialTexture> Textures;
 };
+
+#undef UE_API
